Copying //UE4/Dev-Sequencer to //UE4/Dev-Main (Source: //UE4/Dev-Sequencer @ 3116759)

==========================
MAJOR FEATURES + CHANGES
==========================

Change 2840895 on 2016/01/23 by [REDACTED_PERSON]

	Sequencer: Moved key proxy handling from section into key area; added support for grouped keys; exposed color properties in context menu.

Change 2937981 on 2016/04/08 by [REDACTED_PERSON]

	Sequencer: Refactored GetKeyHandles to take a time range for filtering keys

Change 3058930 on 2016/07/20 by [REDACTED_PERSON]

	Sequencer: Fix crash in undo/redo of creating a sub section.

Change 3062696 on 2016/07/24 by [REDACTED_PERSON]

	Sequencer: Remove curve editor visibility as a toggleable config. It's now just a toggleable state that defaults to false. This fixes unexpected behavior of staying in the curve editor when restarting the
editor or switching to a different level sequence asset.

Change 3062763 on 2016/07/24 by [REDACTED_PERSON]

	Sequencer: Tweak track colors

	Audio track brighter
	Transform, bool, event tracks less saturated
	Recording subsection more saturated
	Fade track gradient

Change 3064262 on 2016/07/26 by [REDACTED_PERSON]

	Sequencer: Add Convert to Possessable

Change 3070101 on 2016/07/29 by [REDACTED_PERSON]

	Sequence Recorder: Record actors as possessables

	Added GetWorld() check when resolving bindings in case the world is being torn down.

Change 3070105 on 2016/07/29 by [REDACTED_PERSON]

	Sequence Recorder: Add cine camera component to default classes and properties to record.

Change 3074750 on 2016/08/02 by [REDACTED_PERSON]

	Sequencer - Simplify keyframing from a mixture of 5 boolean options to the 3 modes we actually use and move some of the logic directly into the keyframe track for consistency.  Also make some keyframing
structs fields constant because public fields lead lead to side effects and difficult debugging.

	This is part of a larger keyframing refactor to fix the inconsistent behavior of default values. I've separated it into 3 separate smaller targeted refactors to make reviews easier which is why I'm
checking in with rb=none.

Change 3075181 on 2016/08/03 by [REDACTED_PERSON]

	Sequencer: Added ability to record audio as part of a recorded sequence

	Implementing the ability to get a global audio clock time per audio device.

	Use FAudioDevice::GetAudioClock() to get the current audio clock (in seconds).

Change 3078922 on 2016/08/05 by [REDACTED_PERSON]

	Sequencer: Spawnables now use deferred spawning

	#2613

Change 3084262 on 2016/08/10 by [REDACTED_PERSON]

	Sequencer - Change the 3d constraint tracks to clear the actors transform themselves instead of the transform track clearing it every frame.

Change 3084851 on 2016/08/10 by [REDACTED_PERSON]

	Sequencer - Make default value handling more consistent
	+ Add an setting for whether or not to automatically set default values for tracks which is enabled for the level sequencer and disabled for UMG.
	+ Add menu item for clearing default values for a track.
	+ Fix property and transform track instances so that they provide the current value when evaluating the section so that empty sections with no defaults don't modify the runtime object.
	+ Change the keyframe track editor so that default values are always set when a transform or property changes if auto set default values is turned on.
	+ Change the key editors so that they always update default values whenever they are changed if auto set default values is turned on.

Change 3089063 on 2016/08/15 by [REDACTED_PERSON]

	Sequence Recorder: Add countdown timer source image. Expand and clamp countdown to 9 seconds.

Change 3091925 on 2016/08/17 by [REDACTED_PERSON]

	Sequencer: Restructure level sequence actor with propert scene component and sprite component attached to it.

	This fixes an ensure when double clicking on the level sequence actor in editor which expects that there is a non editor only root component.

Change 3092201 on 2016/08/17 by [REDACTED_PERSON]

	Sequencer: Attach audio component to actor when creating a spatialized sound. Same thing happens in Matinee.

Change 3092267 on 2016/08/17 by [REDACTED_PERSON]

	Sequencer: Audio spatializiation components should be transient and not owned by the actor it's assigned to.

Change 3098930 on 2016/08/23 by [REDACTED_PERSON]

	Sequencer: Set skeletal meshes to always tick pose and refresh bones in the skeletal animation track instance pre update pass.

Change 3099461 on 2016/08/24 by [REDACTED_PERSON]

	Cine Camera: Add icons for cine camera actor, camera rig rail, camera rig crane.

Change 3099527 on 2016/08/24 by [REDACTED_PERSON]

	Cine Camera: Force the roll and yaw of the camera crane mount to be 0. The yaw is controlled by the crane yaw control.

Change 3099739 on 2016/08/24 by [REDACTED_PERSON]

	Cine Camera: Add toggles to lock the mount pitch/yaw for the crane. By default they are not locked  so the camera will stay level with the ground.

Change 3101574 on 2016/08/25 by [REDACTED_PERSON]

	Sequencer: Validate property path before creating a track. This fixes a bug where two tracks for the same property could be created since property changes could come from the actor or the component.

Change 3102253 on 2016/08/26 by [REDACTED_PERSON]

	Cine Camera: Unregister any owned spline mesh components that aren't being used for visualization. This fixes a bug where undo doesn't match the spline points.

Change 3103891 on 2016/08/27 by [REDACTED_PERSON]

	Movie Capture: Fix game mode override option so that it works in capture in editor.

Change 3104411 on 2016/08/29 by [REDACTED_PERSON]

	Sequencer: Refix attached audio component so that it's owned by the actor it's attached to but still transient.

Change 3105807 on 2016/08/30 by [REDACTED_PERSON]

	Sequencer: Restore state when stopping the level sequence player.

Change 3105988 on 2016/08/30 by [REDACTED_PERSON]

	Sequencer: Fix set key time so that it expands the section range if necessary.

Change 3107109 on 2016/08/30 by [REDACTED_PERSON]

	MediaPlayerEditor: Added support for dragging and dropping media files into media player editor

Change 3107347 on 2016/08/30 by [REDACTED_PERSON]

	Sequencer: Fix crash when key area is null. Also, improved middle mouse drag to create a key mechanism so that the operation is done in one undoable transaction.

Change 3107350 on 2016/08/30 by [REDACTED_PERSON]

	Sequencer: Add option to refresh runtime instances immediately. This fixes a bug when setting keyframes with the key editors - if the instances aren't updated immediately, the wrong/stale value will get
keyed.

Change 3112245 on 2016/09/02 by [REDACTED_PERSON]

	WmfMedia: Fixed WMV2 encoded videos are too bright (UE-35529)

Change 3112621 on 2016/09/02 by [REDACTED_PERSON]

	Engine: Subtitle manager code & documentation cleanup pass

Change 3112712 on 2016/09/02 by [REDACTED_PERSON]

	Matinee: Fixed incorrect subtitle timing & rendering when scrubbing (Github PR# 2591 & 2762)

	https://github.com/EpicGames/UnrealEngine/pull/2591
	https://github.com/EpicGames/UnrealEngine/pull/2762

Change 3113052 on 2016/09/03 by [REDACTED_PERSON]

	Cine Camera: Enable "Actor to Track" as a keyable property.

Change 3114912 on 2016/09/06 by [REDACTED_PERSON]

	AvfMedia: Added missing supported file extensions (UE-35635)

Change 3114951 on 2016/09/06 by [REDACTED_PERSON]

	AvfMedia: Added media source factory file extension descriptions (UE-35635)

Change 3115411 on 2016/09/07 by [REDACTED_PERSON]

	Sequencer: Fix properties not getting bound when converting from matinee to level sequence tracks.

Change 3115412 on 2016/09/07 by [REDACTED_PERSON]

	Editor: Add GEditor check - fix a crash when using sequence recorder with -game.

Change 3115413 on 2016/09/07 by [REDACTED_PERSON]

	Seqeuencer: Put the master sequence dialog's details views in a scrollbox to fix the hidden buttons.

Change 3116156 on 2016/09/07 by [REDACTED_PERSON]

	MediaAssets: Fixed native player not being reused

Change 3116225 on 2016/09/07 by [REDACTED_PERSON]

	MediaAssets: Hooking up caption sink in media player

diff --git a/Engine/Plugins/Editor/AudioCapture/AudioCapture.uplugin b/Engine/Plugins/Editor/AudioCapture/AudioCapture.uplugin
new file mode 100644
--- /dev/null
+++ b/Engine/Plugins/Editor/AudioCapture/AudioCapture.uplugin
@@ -0,0 +1,26 @@
+{
+	"FileVersion" : 3,
+	"Version" : 1,
+	"VersionName" : "1.0",
+	"FriendlyName" : "Audio Capture",
+	"Description" : "Audio capture plugin for sequencer recorder.",
+	"Category" : "Audio",
+	"CreatedBy" : "Epic Games, Inc.",
+	"CreatedByURL" : "",
+	"DocsURL" : "",
+	"MarketplaceURL" : "",
+	"SupportURL" : "",
+	"EnabledByDefault" : true,
+	"CanContainContent" : false,
+	"IsBetaVersion" : true,
+	"Installed" : false,
+	"Modules" :
+	[
+		{
+			"Name" : "AudioCapture",
+			"Type" : "Editor",
+			"LoadingPhase" : "PostEngineInit",
+			"WhitelistPlatforms" : [ "Win64" ]
+		},
+	]
+}
\ No newline at end of file
diff --git a/Engine/Plugins/Editor/AudioCapture/Source/AudioCapture/AudioCapture.build.cs b/Engine/Plugins/Editor/AudioCapture/Source/AudioCapture/AudioCapture.build.cs
new file mode 100644
--- /dev/null
+++ b/Engine/Plugins/Editor/AudioCapture/Source/AudioCapture/AudioCapture.build.cs
@@ -0,0 +1,35 @@
+// Copyright 1998-2016 Epic Games, Inc. All Rights Reserved.
+
+using UnrealBuildTool;
+
+public class AudioCapture : ModuleRules
+{
+	public AudioCapture(TargetInfo Target)
+	{
+		PrivateDependencyModuleNames.AddRange(
+			new string[] {
+				"Core",
+				"CoreUObject",
+				"Engine",
+				"UnrealEd",
+				"SequenceRecorder",
+			}
+		);
+
+		PrivateIncludePaths.AddRange(
+			new string[] {
+				"AudioCapture/Private",
+			}
+		);
+
+        if (Target.Platform == UnrealTargetPlatform.Win32 ||
+            Target.Platform == UnrealTargetPlatform.Win64)
+        {
+            // Add __WINDOWS_WASAPI__ so that RtAudio compiles with WASAPI
+            Definitions.Add("__WINDOWS_DS__");
+
+            // Allow us to use direct sound
+            AddEngineThirdPartyPrivateStaticDependencies(Target, "DirectSound");
+        }
+	}
+}
diff --git a/Engine/Plugins/Editor/AudioCapture/Source/AudioCapture/Private/AudioCapture.cpp b/Engine/Plugins/Editor/AudioCapture/Source/AudioCapture/Private/AudioCapture.cpp
new file mode 100644
--- /dev/null
+++ b/Engine/Plugins/Editor/AudioCapture/Source/AudioCapture/Private/AudioCapture.cpp
@@ -0,0 +1,52 @@
+// Copyright 1998-2016 Epic Games, Inc. All Rights Reserved.
+
+#include "AudioCapture.h"
+#include "ISequenceRecorder.h"
+#include "ISequenceAudioRecorder.h"
+#include "ModuleManager.h"
+#include "AudioRecordingManager.h"
+
+class FAudioRecorder : ISequenceAudioRecorder
+{
+public:
+
+	virtual void Start(const FAudioRecorderSettings& Settings) override
+	{
+		FAudioRecordingManager::Get().StartRecording(Settings.Directory, Settings.AssetName, Settings.RecordingDurationSec, Settings.GainDB, Settings.InputBufferSize);
+	}
+
+	virtual USoundWave* Stop() override
+	{
+		return FAudioRecordingManager::Get().StopRecording();
+	}
+};
+
+
+class FAudioCaptureModule
+	: public IModuleInterface
+{
+	virtual void StartupModule() override
+	{
+		ISequenceRecorder& Recorder = FModuleManager::Get().LoadModuleChecked<ISequenceRecorder>("SequenceRecorder");
+		RecorderHandle = Recorder.RegisterAudioRecorder(
+			[](){
+				return TUniquePtr<ISequenceAudioRecorder>(new FAudioRecorder());
+			}
+		);
+	}
+
+	virtual void ShutdownModule() override
+	{
+		ISequenceRecorder* Recorder = FModuleManager::Get().GetModulePtr<ISequenceRecorder>("SequenceRecorder");
+		
+		if (Recorder)
+		{
+			Recorder->UnregisterAudioRecorder(RecorderHandle);
+		}
+	}
+
+	FDelegateHandle RecorderHandle;
+};
+
+
+IMPLEMENT_PRIMARY_GAME_MODULE( FAudioCaptureModule, AudioCapture, "AudioCapture" );
diff --git a/Engine/Plugins/Editor/AudioCapture/Source/AudioCapture/Private/AudioCapture.h b/Engine/Plugins/Editor/AudioCapture/Source/AudioCapture/Private/AudioCapture.h
new file mode 100644
--- /dev/null
+++ b/Engine/Plugins/Editor/AudioCapture/Source/AudioCapture/Private/AudioCapture.h
@@ -0,0 +1,13 @@
+// Copyright 1998-2016 Epic Games, Inc. All Rights Reserved.
+
+#pragma once
+
+#include "Core.h"
+
+// For some reason this isn't getting set
+#define __WINDOWS_DS__
+
+// And for some reason this is still getting set...
+#ifdef __WINDOWS_WASAPI__
+#undef __WINDOWS_WASAPI__
+#endif
\ No newline at end of file
diff --git a/Engine/Plugins/Editor/AudioCapture/Source/AudioCapture/Private/AudioRecordingManager.cpp b/Engine/Plugins/Editor/AudioCapture/Source/AudioCapture/Private/AudioRecordingManager.cpp
new file mode 100644
--- /dev/null
+++ b/Engine/Plugins/Editor/AudioCapture/Source/AudioCapture/Private/AudioRecordingManager.cpp
@@ -0,0 +1,477 @@
+// Copyright 1998-2016 Epic Games, Inc. All Rights Reserved.
+
+#include "AudioCapture.h"
+#include "AudioRecordingManager.h"
+#include "SoundDefinitions.h"
+#include "Sound/AudioSettings.h"
+#include "AssetRegistryModule.h"
+#include "UnrealEd.h"
+#include "Factories/SoundFactory.h"
+
+
+DEFINE_LOG_CATEGORY(LogMicManager);
+
+/**
+* Callback Function For the Microphone Capture for RtAudio
+*/
+static int32 OnAudioCaptureCallback(void *OutBuffer, void* InBuffer, uint32 InBufferFrames, double StreamTime, RtAudioStreamStatus AudioStreamStatus, void* InUserData)
+{
+	// Check for stream overflow (i.e. we're feeding audio faster than we're consuming)
+
+
+	// Cast the user data to the singleton mic recorder
+	FAudioRecordingManager* AudioRecordingManager = (FAudioRecordingManager*)InUserData;
+
+	// Call the mic capture callback function
+	return AudioRecordingManager->OnAudioCapture(InBuffer, InBufferFrames, StreamTime, AudioStreamStatus == RTAUDIO_INPUT_OVERFLOW);
+}
+
+/**
+* FAudioRecordingManager Implementation
+*/
+
+FAudioRecordingManager::FAudioRecordingManager()
+	: NumRecordedSamples(0)
+	, NumFramesToRecord(0)
+	, RecordingBlockSize(0)
+	, RecordingSampleRate(44100.0f)
+	, NumInputChannels(1)
+	, InputGain(0.0f)
+	, bRecording(false)
+	, bError(false)
+{
+}
+
+FAudioRecordingManager::~FAudioRecordingManager()
+{
+}
+
+FAudioRecordingManager& FAudioRecordingManager::Get()
+{
+	// Static isngleton instance
+	static FAudioRecordingManager MicInputManager;
+	return MicInputManager;
+}
+
+USoundWave* FAudioRecordingManager::StartRecording(const FDirectoryPath& Directory, const FString& AssetName, float RecordingDurationSec, float GainDB, int32 InputBufferSize)
+{
+	if (bError)
+	{
+		return nullptr;
+	}
+
+	// Stop any recordings currently going on (if there is one)
+	USoundWave* NewSoundWave = StopRecording();
+
+	RecordingBlockSize = InputBufferSize;
+
+	// If we have a stream open close it (reusing streams can cause a blip of previous recordings audio)
+	if (ADC.isStreamOpen())
+	{
+		try
+		{
+			ADC.stopStream();
+			ADC.closeStream();
+		}
+		catch (RtAudioError& e)
+		{
+			FString ErrorMessage = FString(e.what());
+			bError = true;
+			UE_LOG(LogMicManager, Error, TEXT("Failed to close the mic capture device stream: %s"), *ErrorMessage);
+			return nullptr;
+		}
+	}
+
+	UE_LOG(LogMicManager, Log, TEXT("Starting mic recording."));
+
+	// Convert input gain decibels into linear scale
+	if (GainDB != 0.0f)
+	{
+		InputGain = FMath::Pow(10.0f, GainDB / 20.0f);
+	}
+	else
+	{
+		InputGain = 0.0f;
+	}
+
+	RtAudio::DeviceInfo Info = ADC.getDeviceInfo(StreamParams.deviceId);
+	RecordingSampleRate = Info.preferredSampleRate;
+	NumInputChannels = Info.inputChannels;
+
+	// Reserve enough space in our current recording buffer for 10 seconds of audio to prevent slowing down due to allocations
+	if (RecordingDurationSec != -1.0f)
+	{
+		int32 NumSamplesToReserve = RecordingDurationSec * RecordingSampleRate * NumInputChannels;
+		CurrentRecordedPCMData.Reset(NumSamplesToReserve);
+	}
+	else
+	{
+		// if not given a duration, resize array to 60 seconds of audio anyway
+		int32 NumSamplesToReserve = 60 * RecordingSampleRate * NumInputChannels;
+		CurrentRecordedPCMData.Reset(NumSamplesToReserve);
+	}
+
+	CurrentRecordingName = AssetName;
+	CurrentRecordingDirectory = Directory;
+
+	NumRecordedSamples = 0;
+	NumOverflowsDetected = 0;
+
+	if (RecordingDurationSec > 0.0f)
+	{
+		NumFramesToRecord = RecordingSampleRate * RecordingDurationSec;
+	}
+
+	// Publish to the mic input thread that we're ready to record...
+	bRecording = true;
+
+	StreamParams.deviceId = ADC.getDefaultInputDevice(); // Only use the default input device for now
+
+	StreamParams.nChannels = NumInputChannels;
+	StreamParams.firstChannel = 0;
+
+	uint32 BufferFrames = FMath::Max(RecordingBlockSize, 256);
+
+	UE_LOG(LogMicManager, Log, TEXT("Initialized mic recording manager at %d hz sample rate, %d channels, and %d Recording Block Size"), RecordingSampleRate, StreamParams.nChannels, BufferFrames);
+
+	// RtAudio uses exceptions for error handling... 
+	try
+	{
+		// Open up new audio stream
+		ADC.openStream(nullptr, &StreamParams, RTAUDIO_SINT16, RecordingSampleRate, &BufferFrames, &OnAudioCaptureCallback, this);
+	}
+	catch (RtAudioError& e)
+	{
+		FString ErrorMessage = FString(e.what());
+		bError = true;
+		UE_LOG(LogMicManager, Error, TEXT("Failed to open the mic capture device: %s"), *ErrorMessage);
+	}
+	catch (const std::exception& e)
+	{
+		bError = true;
+		FString ErrorMessage = FString(e.what());
+		UE_LOG(LogMicManager, Error, TEXT("Failed to open the mic capture device: %s"), *ErrorMessage);
+	}
+	catch (...)
+	{
+		UE_LOG(LogMicManager, Error, TEXT("Failed to open the mic capture device: unknown error"));
+		bError = true;
+	}
+
+	ADC.startStream();
+
+	return NewSoundWave;
+}
+
+static void WriteUInt32ToByteArrayLE(TArray<uint8>& InByteArray, int32& Index, const uint32 Value)
+{
+	InByteArray[Index++] = (uint8)(Value >> 0);
+	InByteArray[Index++] = (uint8)(Value >> 8);
+	InByteArray[Index++] = (uint8)(Value >> 16);
+	InByteArray[Index++] = (uint8)(Value >> 24);
+}
+
+static void WriteUInt16ToByteArrayLE(TArray<uint8>& InByteArray, int32& Index, const uint16 Value)
+{
+	InByteArray[Index++] = (uint8)(Value >> 0);
+	InByteArray[Index++] = (uint8)(Value >> 8);
+}
+
+void FAudioRecordingManager::SerializeWaveFile(TArray<uint8>& OutWaveFileData, const uint8* InPCMData, const int32 NumBytes)
+{
+	// Reserve space for the raw wave data
+	OutWaveFileData.Empty(NumBytes + 44);
+	OutWaveFileData.AddZeroed(NumBytes + 44);
+
+	int32 WaveDataByteIndex = 0;
+
+	// Wave Format Serialization ----------
+
+	// FieldName: ChunkID
+	// FieldSize: 4 bytes
+	// FieldValue: RIFF (FourCC value, big-endian)
+	OutWaveFileData[WaveDataByteIndex++] = 'R';
+	OutWaveFileData[WaveDataByteIndex++] = 'I';
+	OutWaveFileData[WaveDataByteIndex++] = 'F';
+	OutWaveFileData[WaveDataByteIndex++] = 'F';
+
+	// ChunkName: ChunkSize: 4 bytes 
+	// Value: NumBytes + 36. Size of the rest of the chunk following this number. Size of entire file minus 8 bytes.
+	WriteUInt32ToByteArrayLE(OutWaveFileData, WaveDataByteIndex, NumBytes + 36);
+
+	// FieldName: Format 
+	// FieldSize: 4 bytes
+	// FieldValue: "WAVE"  (big-endian)
+	OutWaveFileData[WaveDataByteIndex++] = 'W';
+	OutWaveFileData[WaveDataByteIndex++] = 'A';
+	OutWaveFileData[WaveDataByteIndex++] = 'V';
+	OutWaveFileData[WaveDataByteIndex++] = 'E';
+
+	// FieldName: Subchunk1ID
+	// FieldSize: 4 bytes
+	// FieldValue: "fmt "
+	OutWaveFileData[WaveDataByteIndex++] = 'f';
+	OutWaveFileData[WaveDataByteIndex++] = 'm';
+	OutWaveFileData[WaveDataByteIndex++] = 't';
+	OutWaveFileData[WaveDataByteIndex++] = ' ';
+
+	// FieldName: Subchunk1Size
+	// FieldSize: 4 bytes
+	// FieldValue: 16 for PCM
+	WriteUInt32ToByteArrayLE(OutWaveFileData, WaveDataByteIndex, 16);
+
+	// FieldName: AudioFormat
+	// FieldSize: 2 bytes
+	// FieldValue: 1 for PCM
+	WriteUInt16ToByteArrayLE(OutWaveFileData, WaveDataByteIndex, 1);
+
+	// FieldName: NumChannels
+	// FieldSize: 2 bytes
+	// FieldValue: 1 for for mono
+	WriteUInt16ToByteArrayLE(OutWaveFileData, WaveDataByteIndex, NumInputChannels);
+
+	// FieldName: SampleRate
+	// FieldSize: 4 bytes
+	// FieldValue: MIC_SAMPLE_RATE
+	WriteUInt32ToByteArrayLE(OutWaveFileData, WaveDataByteIndex, WAVE_FILE_SAMPLERATE);
+
+	// FieldName: ByteRate
+	// FieldSize: 4 bytes
+	// FieldValue: SampleRate * NumChannels * BitsPerSample/8
+	int32 ByteRate = WAVE_FILE_SAMPLERATE * NumInputChannels * 2;
+	WriteUInt32ToByteArrayLE(OutWaveFileData, WaveDataByteIndex, ByteRate);
+
+	// FieldName: BlockAlign
+	// FieldSize: 2 bytes
+	// FieldValue: NumChannels * BitsPerSample/8
+	int32 BlockAlign = 2;
+	WriteUInt16ToByteArrayLE(OutWaveFileData, WaveDataByteIndex, BlockAlign);
+
+	// FieldName: BitsPerSample
+	// FieldSize: 2 bytes
+	// FieldValue: 16 (16 bits per sample)
+	WriteUInt16ToByteArrayLE(OutWaveFileData, WaveDataByteIndex, 16);
+
+	// FieldName: Subchunk2ID
+	// FieldSize: 4 bytes
+	// FieldValue: "data" (big endian)
+
+	OutWaveFileData[WaveDataByteIndex++] = 'd';
+	OutWaveFileData[WaveDataByteIndex++] = 'a';
+	OutWaveFileData[WaveDataByteIndex++] = 't';
+	OutWaveFileData[WaveDataByteIndex++] = 'a';
+
+	// FieldName: Subchunk2Size
+	// FieldSize: 4 bytes
+	// FieldValue: number of bytes of the data
+	WriteUInt32ToByteArrayLE(OutWaveFileData, WaveDataByteIndex, NumBytes);
+
+	// Copy the raw PCM data to the audio file
+	FMemory::Memcpy(&OutWaveFileData[WaveDataByteIndex], InPCMData, NumBytes);
+}
+
+static void SampleRateConvert(float CurrentSR, float TargetSR, int32 NumChannels, const TArray<int16>& CurrentRecordedPCMData, int32 NumSamplesToConvert, TArray<int16>& OutConverted)
+{
+	int32 NumInputSamples = CurrentRecordedPCMData.Num();
+	int32 NumOutputSamples = NumInputSamples * TargetSR / CurrentSR;
+
+	OutConverted.Reset(NumOutputSamples);
+
+	float SrFactor = (double)CurrentSR / TargetSR;
+	float CurrentFrameIndexInterpolated = 0.0f;
+	check(NumSamplesToConvert <= CurrentRecordedPCMData.Num());
+	int32 NumFramesToConvert = NumSamplesToConvert / NumChannels;
+	int32 CurrentFrameIndex = 0;
+
+	for (;;)
+	{
+		int32 NextFrameIndex = CurrentFrameIndex + 1;
+		if (CurrentFrameIndex >= NumFramesToConvert || NextFrameIndex >= NumFramesToConvert)
+		{
+			break;
+		}
+
+		for (int32 Channel = 0; Channel < NumChannels; ++Channel)
+		{
+			int32 CurrentSampleIndex = CurrentFrameIndex * NumChannels + Channel;
+			int32 NextSampleIndex = CurrentSampleIndex + NumChannels;
+
+			int16 CurrentSampleValue = CurrentRecordedPCMData[CurrentSampleIndex];
+			int16 NextSampleValue = CurrentRecordedPCMData[NextSampleIndex];
+
+			int16 NewSampleValue = FMath::Lerp(CurrentSampleValue, NextSampleValue, CurrentFrameIndexInterpolated);
+
+			OutConverted.Add(NewSampleValue);
+		}
+
+		CurrentFrameIndexInterpolated += SrFactor;
+
+		// Wrap the interpolated frame between 0.0 and 1.0 to maintain float precision
+		while (CurrentFrameIndexInterpolated >= 1.0f)
+		{
+			CurrentFrameIndexInterpolated -= 1.0f;
+			
+			// Every time it wraps, we increment the frame index
+			++CurrentFrameIndex;
+		}
+	}
+}
+
+USoundWave* FAudioRecordingManager::StopRecording()
+{
+	FScopeLock Lock(&CriticalSection);
+
+	// If we're currently recording, stop the recording
+	if (bRecording)
+	{
+		bRecording = false;
+
+		if (CurrentRecordedPCMData.Num() > 0)
+		{
+			NumRecordedSamples = FMath::Min(NumFramesToRecord * NumInputChannels, CurrentRecordedPCMData.Num());
+
+			UE_LOG(LogMicManager, Log, TEXT("Stopping mic recording. Recorded %d frames of audio (%.4f seconds). Detected %d buffer overflows."), 
+				NumRecordedSamples, 
+				(float)NumRecordedSamples / RecordingSampleRate,
+				NumOverflowsDetected);
+
+			// Get a ptr to the buffer we're actually going to serialize
+			TArray<int16>* PCMDataToSerialize = nullptr;
+
+			// If our sample rate isn't 44100, then we need to do a SRC
+			if (RecordingSampleRate != 44100.0f)
+			{
+				UE_LOG(LogMicManager, Log, TEXT("Converting sample rate from %d hz to 44100 hz."), (int32)RecordingSampleRate);
+
+				SampleRateConvert(RecordingSampleRate, (float)WAVE_FILE_SAMPLERATE, NumInputChannels, CurrentRecordedPCMData, NumRecordedSamples, ConvertedPCMData);
+
+				// Update the recorded samples to the converted buffer samples
+				NumRecordedSamples = ConvertedPCMData.Num();
+				PCMDataToSerialize = &ConvertedPCMData;
+			}
+			else
+			{
+				// Just use the original recorded buffer to serialize
+				PCMDataToSerialize = &CurrentRecordedPCMData;
+			}
+
+			// Scale by the linear gain if it's been set to something (0.0f is ctor default and impossible to set by dB)
+			if (InputGain != 0.0f)
+			{
+				UE_LOG(LogMicManager, Log, TEXT("Scaling gain of recording by %.2f linear gain."), InputGain);
+
+				for (int32 i = 0; i < NumRecordedSamples; ++i)
+				{
+					// Scale by input gain, clamp to prevent integer overflow when casting back to int16. Will still clip.
+					(*PCMDataToSerialize)[i] = (int16)FMath::Clamp(InputGain * (float)(*PCMDataToSerialize)[i], -32767.0f, 32767.0f);
+				}
+			}
+
+			// Get the raw data
+			const uint8* RawData = (const uint8*)PCMDataToSerialize->GetData();
+			int32 NumBytes = NumRecordedSamples * sizeof(int16);
+
+			USoundWave* NewSoundWave = nullptr;
+			USoundWave* ExistingSoundWave = nullptr;
+			TArray<UAudioComponent*> ComponentsToRestart;
+			bool bCreatedPackage = false;
+
+			if (CurrentRecordingDirectory.Path.IsEmpty() || CurrentRecordingName.IsEmpty())
+			{
+				// Create a new sound wave object from the transient package
+				NewSoundWave = NewObject<USoundWave>(GetTransientPackage(), *CurrentRecordingName);
+			}
+			else
+			{
+				// Create a new package
+				FString PackageName = CurrentRecordingDirectory.Path / CurrentRecordingName;
+				UPackage* Package = CreatePackage(nullptr, *PackageName);
+
+				//FEditorDelegates::OnAssetPreImport.Broadcast(this, USoundWave::StaticClass(), Package, *CurrentRecordingName, TEXT("wav"));
+
+				// Create a raw .wav file to stuff the raw PCM data in so when we create the sound wave asset it's identical to a normal imported asset
+				SerializeWaveFile(RawWaveData, RawData, NumBytes);
+
+				// Check to see if a sound wave already exists at this location
+				ExistingSoundWave = FindObject<USoundWave>(Package, *CurrentRecordingName);
+
+				// See if it's currently being played right now
+				FAudioDeviceManager* AudioDeviceManager = GEngine->GetAudioDeviceManager();
+				if (AudioDeviceManager && ExistingSoundWave)
+				{
+					AudioDeviceManager->StopSoundsUsingResource(ExistingSoundWave, &ComponentsToRestart);
+				}
+
+				// Create a new sound wave object
+				if (ExistingSoundWave)
+				{
+					NewSoundWave = ExistingSoundWave;
+					NewSoundWave->FreeResources();
+				}
+				else
+				{
+					NewSoundWave = NewObject<USoundWave>(Package, *CurrentRecordingName, RF_Public | RF_Standalone);
+				}
+
+				// Compressed data is now out of date.
+				NewSoundWave->InvalidateCompressedData();
+
+				// Copy the raw wave data file to the sound wave for storage. Will allow the recording to be exported.
+				NewSoundWave->RawData.Lock(LOCK_READ_WRITE);
+				void* LockedData = NewSoundWave->RawData.Realloc(RawWaveData.Num());
+				FMemory::Memcpy(LockedData, RawWaveData.GetData(), RawWaveData.Num());
+				NewSoundWave->RawData.Unlock();
+
+				bCreatedPackage = true;
+			}
+
+			if (NewSoundWave)
+			{
+				// Copy the recorded data to the sound wave so we can quickly preview it
+				NewSoundWave->RawPCMDataSize = NumBytes;
+				NewSoundWave->RawPCMData = (uint8*)FMemory::Malloc(NewSoundWave->RawPCMDataSize);
+				FMemory::Memcpy(NewSoundWave->RawPCMData, RawData, NumBytes);
+
+				// Calculate the duration of the sound wave
+				NewSoundWave->Duration = (float)(NumRecordedSamples / NumInputChannels) / WAVE_FILE_SAMPLERATE;
+				NewSoundWave->SampleRate = WAVE_FILE_SAMPLERATE;
+				NewSoundWave->NumChannels = NumInputChannels;
+
+				if (bCreatedPackage)
+				{
+					//FEditorDelegates::OnAssetPostImport.Broadcast(this, NewSoundWave);
+
+					FAssetRegistryModule::AssetCreated(NewSoundWave);
+					NewSoundWave->MarkPackageDirty();
+
+					// Restart any audio components if they need to be restarted
+					for (int32 ComponentIndex = 0; ComponentIndex < ComponentsToRestart.Num(); ++ComponentIndex)
+					{
+						ComponentsToRestart[ComponentIndex]->Play();
+					}
+				}
+			}
+
+			return NewSoundWave;
+		}
+	}
+
+
+	return nullptr;
+}
+
+int32 FAudioRecordingManager::OnAudioCapture(void* InBuffer, uint32 InBufferFrames, double StreamTime, bool bOverflow)
+{
+	FScopeLock Lock(&CriticalSection);
+
+	if (bRecording)
+	{
+		if (bOverflow) 
+			++NumOverflowsDetected;
+
+		CurrentRecordedPCMData.Append((int16*)InBuffer, InBufferFrames * NumInputChannels);
+		return 0;
+	}
+
+	return 1;
+}
\ No newline at end of file
diff --git a/Engine/Plugins/Editor/AudioCapture/Source/AudioCapture/Private/AudioRecordingManager.h b/Engine/Plugins/Editor/AudioCapture/Source/AudioCapture/Private/AudioRecordingManager.h
new file mode 100644
--- /dev/null
+++ b/Engine/Plugins/Editor/AudioCapture/Source/AudioCapture/Private/AudioRecordingManager.h
@@ -0,0 +1,95 @@
+// Copyright 1998-2016 Epic Games, Inc. All Rights Reserved.
+
+#pragma once
+
+#include "RtAudio.h"
+#include "UnrealEd.h"
+
+DECLARE_LOG_CATEGORY_EXTERN(LogMicManager, Log, All);
+
+/**
+ * FAudioRecordingManager
+ * Singleton Mic Recording Manager -- generates recordings, stores the recorded data and plays them back
+ */
+class FAudioRecordingManager
+{
+public:
+	// Retrieves the singleton recording manager
+	static FAudioRecordingManager& Get();
+
+	// Starts a new recording with the given name and optional duration. 
+	// If set to -1.0f, a duration won't be used and the recording length will be determined by StopRecording().
+	USoundWave* StartRecording(const FDirectoryPath& Directory, const FString& AssetName, float RecordingDurationSec, float GainDB, int32 InputBufferSize);
+
+	// Stops recording if the recording manager is recording. If not recording but has recorded data (due to set duration), it will just return the generated USoundWave.
+	USoundWave* StopRecording();
+
+	// Called by RtAudio when a new audio buffer is ready to be supplied.
+	int32 OnAudioCapture(void* InBuffer, uint32 InBufferFrames, double StreamTime, bool bOverflow);
+
+private:
+
+	// Private Constructor
+	FAudioRecordingManager();
+
+	// Private Destructor
+	~FAudioRecordingManager();
+
+	// Saves raw PCM data recorded to a wave file format
+	void SerializeWaveFile(TArray<uint8>& OutWaveFileData, const uint8* InPCMData, const int32 NumBytes);
+
+	// The default mic sample rate
+	const int32 WAVE_FILE_SAMPLERATE = 44100;
+
+	// RtAudio ADC object -- used to interact with low-level audio device.
+	RtAudio ADC;
+
+	// Stream parameters to initialize the ADC
+	RtAudio::StreamParameters StreamParams;
+
+	// Critical section used to stop and retrieve finished audio buffers.
+	FCriticalSection CriticalSection;
+
+	// The current recording name that is actively recording. Empty string if nothing is currently recording.
+	FString CurrentRecordingName;
+
+	// Where to store the current recording.
+	FDirectoryPath CurrentRecordingDirectory;
+
+	// The data which is currently being recorded to, if the manager is actively recording. This is not safe to access while recording.
+	TArray<int16> CurrentRecordedPCMData;
+
+	// Buffer to store sample rate converted PCM data
+	TArray<int16> ConvertedPCMData;
+
+	// Reusable raw wave data buffer to generate .wav file formats
+	TArray<uint8> RawWaveData;
+
+	// The number of frames that have been recorded
+	int32 NumRecordedSamples;
+
+	// The number of frames to record if recording a set duration
+	int32 NumFramesToRecord;
+
+	// Recording block size (number of frames per callback block)
+	int32 RecordingBlockSize;
+
+	// The sample rate used in the recording
+	float RecordingSampleRate;
+
+	// Num input channels
+	int32 NumInputChannels;
+
+	// A linear gain to apply on mic input
+	float InputGain;
+
+	// Whether or not the manager is actively recording.
+	FThreadSafeBool bRecording;
+
+	// Number of overflows detected while recording
+	int32 NumOverflowsDetected;
+
+	// Whether or not we have an error
+	uint32 bError : 1;
+
+};
