[Release 4.20] fix crash that occurs when Audio Capture components are started a second time. #fyi dan.reynolds, aaron.mcleran #rb none #jira none

diff --git a/Engine/Plugins/Runtime/AudioCapture/Source/AudioCapture/Private/AudioCaptureComponent.cpp b/Engine/Plugins/Runtime/AudioCapture/Source/AudioCapture/Private/AudioCaptureComponent.cpp
--- a/Engine/Plugins/Runtime/AudioCapture/Source/AudioCapture/Private/AudioCaptureComponent.cpp
+++ b/Engine/Plugins/Runtime/AudioCapture/Source/AudioCapture/Private/AudioCaptureComponent.cpp
@@ -58,6 +58,13 @@ bool UAudioCaptureComponent::IsReadyForFinishDestroy()
 
 void UAudioCaptureComponent::FinishDestroy()
 {
+	if (CaptureSynth.IsStreamOpen())
+	{
+		CaptureSynth.AbortCapturing();
+	}
+
+	check(!CaptureSynth.IsStreamOpen());
+
 	Super::FinishDestroy();
 	bSuccessfullyInitialized = false;
 	bIsCapturing = false;
@@ -67,21 +74,20 @@ void UAudioCaptureComponent::FinishDestroy()
 
 void UAudioCaptureComponent::OnBeginGenerate()
 {
-	if (bIsStreamOpen)
+	if (!bIsStreamOpen)
 	{
-		CaptureSynth.StartCapturing();
-		check(CaptureSynth.IsCapturing());
+		CaptureSynth.OpenDefaultStream();
+	}
 
-		// Don't allow this component to be destroyed until the stream is closed again
-		bIsReadyForForFinishDestroy = false;
+	CaptureSynth.StartCapturing();
+	check(CaptureSynth.IsCapturing());
+
+	// Don't allow this component to be destroyed until the stream is closed again
+	bIsReadyForForFinishDestroy = false;
+
+	FramesSinceStarting = 0;
+	ReadSampleIndex = 0;
 
-		FramesSinceStarting = 0;
-		ReadSampleIndex = 0;
-	}
-	else
-	{
-		UE_LOG(LogAudio, Warning, TEXT("Was not able to open a default capture stream."));
-	}
 }
 
 void UAudioCaptureComponent::OnEndGenerate()
@@ -89,8 +95,7 @@ void UAudioCaptureComponent::OnEndGenerate()
 	if (bIsStreamOpen)
 	{
 		check(CaptureSynth.IsStreamOpen());
-		CaptureSynth.AbortCapturing();
-		check(!CaptureSynth.IsStreamOpen());
+		CaptureSynth.StopCapturing();
 
 		bIsReadyForForFinishDestroy = true;
 	}
