Merging CL 4300421 from Dev-VR to Main - Changing includes to make sure PLATFORM_WINDOWS is defined for non-unity Audio Capture plugins

diff --git a/Engine/Plugins/Runtime/AudioCapture/Source/AudioCapture/Private/AudioCaptureRtAudio.cpp b/Engine/Plugins/Runtime/AudioCapture/Source/AudioCapture/Private/AudioCaptureRtAudio.cpp
--- a/Engine/Plugins/Runtime/AudioCapture/Source/AudioCapture/Private/AudioCaptureRtAudio.cpp
+++ b/Engine/Plugins/Runtime/AudioCapture/Source/AudioCapture/Private/AudioCaptureRtAudio.cpp
@@ -1,5 +1,7 @@
 // Copyright 1998-2018 Epic Games, Inc. All Rights Reserved.
 
+#include "AudioCapture.h"
+
 #if PLATFORM_WINDOWS
 
 #include "AudioCaptureInternal.h"
diff --git a/Engine/Plugins/Runtime/AudioCapture/Source/AudioCapture/Private/RtAudio.cpp b/Engine/Plugins/Runtime/AudioCapture/Source/AudioCapture/Private/RtAudio.cpp
--- a/Engine/Plugins/Runtime/AudioCapture/Source/AudioCapture/Private/RtAudio.cpp
+++ b/Engine/Plugins/Runtime/AudioCapture/Source/AudioCapture/Private/RtAudio.cpp
@@ -1,4 +1,6 @@
-﻿#if PLATFORM_WINDOWS
+﻿#include "RtAudio.h"
+
+#if PLATFORM_WINDOWS
 /************************************************************************/
 /*! \class RtAudio
     \brief Realtime audio i/o C++ classes.
@@ -41,7 +43,6 @@
 
 // RtAudio: Version 5.0.0
 
-#include "RtAudio.h"
 #include "AudioCapture.h"
 THIRD_PARTY_INCLUDES_START
 #include <iostream>
diff --git a/Engine/Plugins/Runtime/AudioCapture/Source/AudioCapture/Private/RtAudio.h b/Engine/Plugins/Runtime/AudioCapture/Source/AudioCapture/Private/RtAudio.h
--- a/Engine/Plugins/Runtime/AudioCapture/Source/AudioCapture/Private/RtAudio.h
+++ b/Engine/Plugins/Runtime/AudioCapture/Source/AudioCapture/Private/RtAudio.h
@@ -1,3 +1,5 @@
+#include "AudioCapture.h"
+
 #if PLATFORM_WINDOWS
 /************************************************************************/
 /*! \class RtAudio
