[NO NON-PII COMMIT MESSAGE]

diff --git a/Engine/Plugins/Runtime/AudioCapture/Source/AudioCapture/Private/AudioCaptureInternal.h b/Engine/Plugins/Runtime/AudioCapture/Source/AudioCapture/Private/AudioCaptureInternal.h
--- a/Engine/Plugins/Runtime/AudioCapture/Source/AudioCapture/Private/AudioCaptureInternal.h
+++ b/Engine/Plugins/Runtime/AudioCapture/Source/AudioCapture/Private/AudioCaptureInternal.h
@@ -79,7 +79,7 @@ namespace Audio
 	};
 
 	// Return nullptr when creating impl
-	TUniquePtr<FAudioCaptureImpl> FAudioCapture::CreateImpl()
+	FORCEINLINE TUniquePtr<FAudioCaptureImpl> FAudioCapture::CreateImpl()
 	{
 		return nullptr;
 	}
