Audio Recording: Fix log output for duration

diff --git a/Engine/Plugins/Runtime/AudioCapture/Source/AudioCaptureEditor/Private/AudioRecordingManager.cpp b/Engine/Plugins/Runtime/AudioCapture/Source/AudioCaptureEditor/Private/AudioRecordingManager.cpp
--- a/Engine/Plugins/Runtime/AudioCapture/Source/AudioCaptureEditor/Private/AudioRecordingManager.cpp
+++ b/Engine/Plugins/Runtime/AudioCapture/Source/AudioCaptureEditor/Private/AudioRecordingManager.cpp
@@ -252,7 +252,7 @@ void FAudioRecordingManager::StopRecording(TArray<USoundWave*>& OutSoundWaves)
 
 			UE_LOG(LogMicManager, Log, TEXT("Stopping mic recording. Recorded %d frames of audio (%.4f seconds). Detected %d buffer overflows."), 
 				NumRecordedSamples, 
-				(float)NumRecordedSamples / RecordingSampleRate,
+				((float)NumRecordedSamples / RecordingSampleRate) / (float)NumInputChannels,
 				NumOverflowsDetected);
 
 			// Get a ptr to the buffer we're actually going to serialize
