Merging CL 4175640 (fix for Audio Capture pausing and resuming) to Dev-VR

diff --git a/Engine/Plugins/Runtime/AudioCapture/Source/AudioCapture/Private/AudioCaptureInternal.h b/Engine/Plugins/Runtime/AudioCapture/Source/AudioCapture/Private/AudioCaptureInternal.h
--- a/Engine/Plugins/Runtime/AudioCapture/Source/AudioCapture/Private/AudioCaptureInternal.h
+++ b/Engine/Plugins/Runtime/AudioCapture/Source/AudioCapture/Private/AudioCaptureInternal.h
@@ -40,6 +40,10 @@ namespace Audio
 		TArray<float> FloatBuffer;
 		MLHandle InputDeviceHandle;
 		bool bStreamStarted;
+
+		FCriticalSection ApplicationResumeCriticalSection;
+		void OnApplicationSuspend();
+		void OnApplicationResume();
 #endif
 
 	private:
diff --git a/Engine/Plugins/Runtime/AudioCapture/Source/AudioCapture/Private/AudioCaptureMLAudio.cpp b/Engine/Plugins/Runtime/AudioCapture/Source/AudioCapture/Private/AudioCaptureMLAudio.cpp
--- a/Engine/Plugins/Runtime/AudioCapture/Source/AudioCapture/Private/AudioCaptureMLAudio.cpp
+++ b/Engine/Plugins/Runtime/AudioCapture/Source/AudioCapture/Private/AudioCaptureMLAudio.cpp
@@ -1,6 +1,8 @@
 // Copyright 1998-2018 Epic Games, Inc. All Rights Reserved.
 
 #include "AudioCaptureInternal.h"
+#include "Misc/CoreDelegates.h"
+
 #if PLATFORM_LUMIN
 
 namespace Audio 
@@ -41,6 +43,8 @@ void FAudioCaptureImpl::OnAudioCapture(void* InBuffer, uint32 InBufferFrames, do
 {
 	check(Callback);
 
+	FScopeLock ScopeLock(&ApplicationResumeCriticalSection);
+
 	int32 NumSamples = (int32)(InBufferFrames * NumChannels);
 
 	//TODO: Check to see if float conversion is needed:
@@ -98,6 +102,9 @@ bool FAudioCaptureImpl::OpenDefaultCaptureStream(const FAudioCaptureStreamParam&
 		return false;
 	}
 
+	FCoreDelegates::ApplicationWillEnterBackgroundDelegate.AddRaw(this, &FAudioCapture::OnApplicationSuspend);
+	FCoreDelegates::ApplicationHasEnteredForegroundDelegate.AddRaw(this, &FAudioCapture::OnApplicationResume);
+
 	// set up variables to be populated by ML Audio
 	uint32 ChannelCount = NumChannels;
 	MLAudioBufferFormat DefaultBufferFormat;
@@ -152,6 +159,10 @@ bool FAudioCaptureImpl::CloseStream()
 	}
 
 	InputDeviceHandle = ML_INVALID_HANDLE;
+
+	FCoreDelegates::ApplicationHasEnteredForegroundDelegate.RemoveAll(this);
+	FCoreDelegates::ApplicationWillEnterBackgroundDelegate.RemoveAll(this);
+
 	return true;
 }
 
@@ -212,6 +223,27 @@ TUniquePtr<FAudioCaptureImpl> FAudioCapture::CreateImpl()
 	return TUniquePtr<FAudioCaptureImpl>(new FAudioCaptureImpl());
 }
 
+void FAudioCapture::OnApplicationSuspend()
+{
+	FScopeLock ScopeLock(&ApplicationResumeCriticalSection);
+	MLResult Result = MLAudioStopInput(InputDeviceHandle);
+	if (Result != MLResult_Ok)
+	{
+		UE_LOG(LogAudioCapture, Warning, TEXT("MLAudioStopInput failed with code %d"), Result);
+		return;
+	}
+}
+
+void FAudioCapture::OnApplicationResume()
+{
+	MLResult Result = MLAudioStartInput(InputDeviceHandle);
+	if (Result != MLResult_Ok)
+	{
+		UE_LOG(LogAudioCapture, Warning, TEXT("MLAudioStartInput failed with code %d"), Result);
+		return;
+	}
+}
+
 } // namespace Audio
 
 #endif // PLATFORM_LUMIN
