Shelve 4175638 got committed but didn't compile.  Fixed compile errors and changed some checks from Handle != ML_HANDLE_INVALID to MlIsValidHandle(Handle), fixed functions to return false if they error, responded to the errors by not continuing further, etc...

Don't know if this fixes all the functionality, but doesn't crash for me anymore.

diff --git a/Engine/Plugins/Runtime/AudioCapture/Source/AudioCapture/Private/AudioCaptureComponent.cpp b/Engine/Plugins/Runtime/AudioCapture/Source/AudioCapture/Private/AudioCaptureComponent.cpp
--- a/Engine/Plugins/Runtime/AudioCapture/Source/AudioCapture/Private/AudioCaptureComponent.cpp
+++ b/Engine/Plugins/Runtime/AudioCapture/Source/AudioCapture/Private/AudioCaptureComponent.cpp
@@ -76,17 +76,19 @@ void UAudioCaptureComponent::OnBeginGenerate()
 {
 	if (!bIsStreamOpen)
 	{
-		CaptureSynth.OpenDefaultStream();
+		bIsStreamOpen = CaptureSynth.OpenDefaultStream();
 	}
 
-	CaptureSynth.StartCapturing();
-	check(CaptureSynth.IsCapturing());
-
-	// Don't allow this component to be destroyed until the stream is closed again
-	bIsReadyForForFinishDestroy = false;
+	if (bIsStreamOpen)
+	{
+		CaptureSynth.StartCapturing();
+		check(CaptureSynth.IsCapturing());
 
-	FramesSinceStarting = 0;
-	ReadSampleIndex = 0;
+		// Don't allow this component to be destroyed until the stream is closed again
+		bIsReadyForForFinishDestroy = false;
+		FramesSinceStarting = 0;
+		ReadSampleIndex = 0;
+	}
 
 }
 
diff --git a/Engine/Plugins/Runtime/AudioCapture/Source/AudioCapture/Private/AudioCaptureInternal.h b/Engine/Plugins/Runtime/AudioCapture/Source/AudioCapture/Private/AudioCaptureInternal.h
--- a/Engine/Plugins/Runtime/AudioCapture/Source/AudioCapture/Private/AudioCaptureInternal.h
+++ b/Engine/Plugins/Runtime/AudioCapture/Source/AudioCapture/Private/AudioCaptureInternal.h
@@ -40,6 +40,10 @@ namespace Audio
 		TArray<float> FloatBuffer;
 		MLHandle InputDeviceHandle;
 		bool bStreamStarted;
+
+		FCriticalSection ApplicationResumeCriticalSection;
+		void OnApplicationSuspend();
+		void OnApplicationResume();
 #endif
 
 	private:
diff --git a/Engine/Plugins/Runtime/AudioCapture/Source/AudioCapture/Private/AudioCaptureMLAudio.cpp b/Engine/Plugins/Runtime/AudioCapture/Source/AudioCapture/Private/AudioCaptureMLAudio.cpp
--- a/Engine/Plugins/Runtime/AudioCapture/Source/AudioCapture/Private/AudioCaptureMLAudio.cpp
+++ b/Engine/Plugins/Runtime/AudioCapture/Source/AudioCapture/Private/AudioCaptureMLAudio.cpp
@@ -1,6 +1,8 @@
 // Copyright 1998-2018 Epic Games, Inc. All Rights Reserved.
 
 #include "AudioCaptureInternal.h"
+#include "Misc/CoreDelegates.h"
+
 #if PLATFORM_LUMIN
 
 namespace Audio 
@@ -41,6 +43,8 @@ void FAudioCaptureImpl::OnAudioCapture(void* InBuffer, uint32 InBufferFrames, do
 {
 	check(Callback);
 
+	FScopeLock ScopeLock(&ApplicationResumeCriticalSection);
+
 	int32 NumSamples = (int32)(InBufferFrames * NumChannels);
 
 	//TODO: Check to see if float conversion is needed:
@@ -78,16 +82,16 @@ bool FAudioCaptureImpl::GetDefaultCaptureDeviceInfo(FCaptureDeviceInfo& OutInfo)
 	else
 	{
 		UE_LOG(LogAudioCapture, Error, TEXT("Unable to retrieve settings from MLAudio error: %d! Sample rate %u"), Result, UnsignedSampleRate);
+		return false;
 	}
 	
-
 	return true;
 }
 
 bool FAudioCaptureImpl::OpenDefaultCaptureStream(const FAudioCaptureStreamParam& StreamParams)
 {
 	UE_LOG(LogAudioCapture, Warning, TEXT("OPENING STREAM"));
-	if (InputDeviceHandle != ML_INVALID_HANDLE)
+	if (MLHandleIsValid(InputDeviceHandle))
 	{
 		UE_LOG(LogAudioCapture, Error, TEXT("Capture Stream Already Opened"));
 	}
@@ -98,6 +102,9 @@ bool FAudioCaptureImpl::OpenDefaultCaptureStream(const FAudioCaptureStreamParam&
 		return false;
 	}
 
+	FCoreDelegates::ApplicationWillEnterBackgroundDelegate.AddRaw(this, &FAudioCaptureImpl::OnApplicationSuspend);
+	FCoreDelegates::ApplicationHasEnteredForegroundDelegate.AddRaw(this, &FAudioCaptureImpl::OnApplicationResume);
+
 	// set up variables to be populated by ML Audio
 	uint32 ChannelCount = NumChannels;
 	MLAudioBufferFormat DefaultBufferFormat;
@@ -133,10 +140,12 @@ bool FAudioCaptureImpl::OpenDefaultCaptureStream(const FAudioCaptureStreamParam&
 	if (Result != MLResult_Ok)
 	{
 		UE_LOG(LogAudioCapture, Warning, TEXT("MLAudioCreateInputFromVoiceComm failed with code %d"), Result);
+		return false;
 	}
-	else if (InputDeviceHandle == ML_INVALID_HANDLE)
+	else if (!MLHandleIsValid(InputDeviceHandle))
 	{
 		UE_LOG(LogAudioCapture, Warning, TEXT("MLAudioCreateInputFromVoiceComm failed to generate an input device handle."));
+		return false;
 	}
 	return true;
 }
@@ -152,6 +161,10 @@ bool FAudioCaptureImpl::CloseStream()
 	}
 
 	InputDeviceHandle = ML_INVALID_HANDLE;
+
+	FCoreDelegates::ApplicationHasEnteredForegroundDelegate.RemoveAll(this);
+	FCoreDelegates::ApplicationWillEnterBackgroundDelegate.RemoveAll(this);
+
 	return true;
 }
 
@@ -199,7 +212,7 @@ bool FAudioCaptureImpl::GetStreamTime(double& OutStreamTime)
 
 bool FAudioCaptureImpl::IsStreamOpen() const
 {
-	return InputDeviceHandle != ML_INVALID_HANDLE;
+	return MLHandleIsValid(InputDeviceHandle);
 }
 
 bool FAudioCaptureImpl::IsCapturing() const
@@ -212,6 +225,27 @@ TUniquePtr<FAudioCaptureImpl> FAudioCapture::CreateImpl()
 	return TUniquePtr<FAudioCaptureImpl>(new FAudioCaptureImpl());
 }
 
+void FAudioCaptureImpl::OnApplicationSuspend()
+{
+	FScopeLock ScopeLock(&ApplicationResumeCriticalSection);
+	MLResult Result = MLAudioStopInput(InputDeviceHandle);
+	if (Result != MLResult_Ok)
+	{
+		UE_LOG(LogAudioCapture, Warning, TEXT("MLAudioStopInput failed with code %d"), Result);
+		return;
+	}
+}
+
+void FAudioCaptureImpl::OnApplicationResume()
+{
+	MLResult Result = MLAudioStartInput(InputDeviceHandle);
+	if (Result != MLResult_Ok)
+	{
+		UE_LOG(LogAudioCapture, Warning, TEXT("MLAudioStartInput failed with code %d"), Result);
+		return;
+	}
+}
+
 } // namespace Audio
 
 #endif // PLATFORM_LUMIN
