#DEV-VR - Fixing Audio Capture to not re-open stream if it's already open

diff --git a/Engine/Plugins/Runtime/AudioCapture/Source/AudioCapture/Private/AudioCapture.cpp b/Engine/Plugins/Runtime/AudioCapture/Source/AudioCapture/Private/AudioCapture.cpp
--- a/Engine/Plugins/Runtime/AudioCapture/Source/AudioCapture/Private/AudioCapture.cpp
+++ b/Engine/Plugins/Runtime/AudioCapture/Source/AudioCapture/Private/AudioCapture.cpp
@@ -1,6 +1,7 @@
 // Copyright 1998-2018 Epic Games, Inc. All Rights Reserved.
 
 #include "AudioCapture.h"
+#include "CoreMinimal.h"
 #include "Misc/ScopeLock.h"
 #include "Modules/ModuleManager.h"
 #include "AudioCaptureInternal.h"
@@ -144,22 +145,27 @@ bool FAudioCaptureSynth::GetDefaultCaptureDeviceInfo(FCaptureDeviceInfo& OutInfo
 
 bool FAudioCaptureSynth::OpenDefaultStream()
 {
-	check(!AudioCapture.IsStreamOpen());
-
-	FAudioCaptureStreamParam StreamParam;
-	StreamParam.Callback = this;
-	StreamParam.NumFramesDesired = 1024;
-
-	// Prepare the audio buffer memory for 2 seconds of stereo audio at 48k SR to reduce chance for allocation in callbacks
-	AudioCaptureData.Reserve(2 * 2 * 48000);
-
-	// Start the stream here to avoid hitching the audio render thread. 
-	if (AudioCapture.OpenDefaultCaptureStream(StreamParam))
+	bool bSuccess = true;
+	if (!AudioCapture.IsStreamOpen())
 	{
-		AudioCapture.StartStream();
-		return true;
+		FAudioCaptureStreamParam StreamParam;
+		StreamParam.Callback = this;
+		StreamParam.NumFramesDesired = 1024;
+
+		// Prepare the audio buffer memory for 2 seconds of stereo audio at 48k SR to reduce chance for allocation in callbacks
+		AudioCaptureData.Reserve(2 * 2 * 48000);
+
+		// Start the stream here to avoid hitching the audio render thread. 
+		if (AudioCapture.OpenDefaultCaptureStream(StreamParam))
+		{
+			AudioCapture.StartStream();
+		}
+		else
+		{
+			bSuccess = false;
+		}
 	}
-	return false;
+	return bSuccess;
 }
 
 bool FAudioCaptureSynth::StartCapturing()
