Remove unneeded whitelisting from AudioCapture plugin

#android
#ios
#ue4

diff --git a/Engine/Plugins/Runtime/AudioCapture/AudioCapture.uplugin b/Engine/Plugins/Runtime/AudioCapture/AudioCapture.uplugin
--- a/Engine/Plugins/Runtime/AudioCapture/AudioCapture.uplugin
+++ b/Engine/Plugins/Runtime/AudioCapture/AudioCapture.uplugin
@@ -19,14 +19,12 @@
 		{
 			"Name" : "AudioCapture",
 			"Type" : "Runtime",
-      "LoadingPhase": "Default",
- 			"WhitelistPlatforms" : [ "Win64", "Lumin", "Mac" ]
-		},
+			"LoadingPhase": "Default"
+ 		},
 		{
 			"Name" : "AudioCaptureEditor",
 			"Type" : "Editor",
-      "LoadingPhase": "PostEngineInit",
- 			"WhitelistPlatforms" : [ "Win64" ]
+			"LoadingPhase": "PostEngineInit"
 		}
 	]
 }
\ No newline at end of file
