Revert to whitelisting only Win64 for AudioCaptureEditor

AudioCaptureEditor needs support for Linux/Mac (Failing to build from source)

diff --git a/Engine/Plugins/Runtime/AudioCapture/AudioCapture.uplugin b/Engine/Plugins/Runtime/AudioCapture/AudioCapture.uplugin
--- a/Engine/Plugins/Runtime/AudioCapture/AudioCapture.uplugin
+++ b/Engine/Plugins/Runtime/AudioCapture/AudioCapture.uplugin
@@ -24,7 +24,8 @@
 		{
 			"Name" : "AudioCaptureEditor",
 			"Type" : "Editor",
-			"LoadingPhase": "PostEngineInit"
+			"LoadingPhase": "PostEngineInit",
+ 			"WhitelistPlatforms" : [ "Win64" ]
 		}
 	]
-}
\ No newline at end of file
+}
