Fix for hot reload hang, where the FVTableHelper constructor wasn't defined, leaving the bIsReadyForForFinishDestroy set to false.

diff --git a/Engine/Plugins/Runtime/AudioCapture/Source/AudioCapture/Private/AudioCaptureComponent.cpp b/Engine/Plugins/Runtime/AudioCapture/Source/AudioCapture/Private/AudioCaptureComponent.cpp
--- a/Engine/Plugins/Runtime/AudioCapture/Source/AudioCapture/Private/AudioCaptureComponent.cpp
+++ b/Engine/Plugins/Runtime/AudioCapture/Source/AudioCapture/Private/AudioCaptureComponent.cpp
@@ -10,7 +10,7 @@ UAudioCaptureComponent::UAudioCaptureComponent(const FObjectInitializer& ObjectI
 	CapturedAudioDataSamples = 0;
 	ReadSampleIndex = 0;
 	bIsDestroying = false;
-	bIsReadyForForFinishDestroy = true;
+	bIsNotReadyForForFinishDestroy = false;
 	bIsStreamOpen = false;
 	CaptureAudioData.Reserve(2 * 48000 * 5);
 }
@@ -53,7 +53,7 @@ void UAudioCaptureComponent::BeginDestroy()
 
 bool UAudioCaptureComponent::IsReadyForFinishDestroy()
 {
-	return bIsReadyForForFinishDestroy;
+	return !bIsNotReadyForForFinishDestroy;
 }
 
 void UAudioCaptureComponent::FinishDestroy()
@@ -85,7 +85,7 @@ void UAudioCaptureComponent::OnBeginGenerate()
 		check(CaptureSynth.IsCapturing());
 
 		// Don't allow this component to be destroyed until the stream is closed again
-		bIsReadyForForFinishDestroy = false;
+		bIsNotReadyForForFinishDestroy = true;
 		FramesSinceStarting = 0;
 		ReadSampleIndex = 0;
 	}
@@ -100,7 +100,7 @@ void UAudioCaptureComponent::OnEndGenerate()
 		CaptureSynth.StopCapturing();
 		bIsStreamOpen = false;
 
-		bIsReadyForForFinishDestroy = true;
+		bIsNotReadyForForFinishDestroy = false;
 	}
 }
 
diff --git a/Engine/Plugins/Runtime/AudioCapture/Source/AudioCapture/Public/AudioCaptureComponent.h b/Engine/Plugins/Runtime/AudioCapture/Source/AudioCapture/Public/AudioCaptureComponent.h
--- a/Engine/Plugins/Runtime/AudioCapture/Source/AudioCapture/Public/AudioCaptureComponent.h
+++ b/Engine/Plugins/Runtime/AudioCapture/Source/AudioCapture/Public/AudioCaptureComponent.h
@@ -52,5 +52,5 @@ private:
 	int32 FramesSinceStarting;
 	int32 ReadSampleIndex;
 	FThreadSafeBool bIsDestroying;
-	FThreadSafeBool bIsReadyForForFinishDestroy;
+	FThreadSafeBool bIsNotReadyForForFinishDestroy;
 };
\ No newline at end of file
