Attempted fix to build error.

diff --git a/Engine/Plugins/Runtime/AudioCapture/Source/AudioCapture/Private/AudioCaptureMLAudio.cpp b/Engine/Plugins/Runtime/AudioCapture/Source/AudioCapture/Private/AudioCaptureMLAudio.cpp
--- a/Engine/Plugins/Runtime/AudioCapture/Source/AudioCapture/Private/AudioCaptureMLAudio.cpp
+++ b/Engine/Plugins/Runtime/AudioCapture/Source/AudioCapture/Private/AudioCaptureMLAudio.cpp
@@ -9,8 +9,7 @@ namespace Audio
 {
 
 FAudioCaptureImpl::FAudioCaptureImpl()
-	: Callback(nullptr)
-	, NumChannels(1)
+	: NumChannels(1)
 	, SampleRate(16000)
 	, InputDeviceHandle(ML_INVALID_HANDLE)
 	, bStreamStarted(false)
@@ -41,7 +40,6 @@ static void OnAudioCaptureCallback(MLHandle Handle, void* CallbackContext)
 
 void FAudioCaptureImpl::OnAudioCapture(void* InBuffer, uint32 InBufferFrames, double StreamTime, bool bOverflow)
 {
-
 	FScopeLock ScopeLock(&ApplicationResumeCriticalSection);
 	int32 NumSamples = (int32)(InBufferFrames * NumChannels);
 
@@ -57,7 +55,7 @@ void FAudioCaptureImpl::OnAudioCapture(void* InBuffer, uint32 InBufferFrames, do
 		FloatBufferPtr[i] = (((float)InBufferData[i]) / 32767.0f);
 	};
 
-	OnAudioCapture(FloatBufferPtr, InBufferFrames, NumChannels, StreamTime, bOverflow);
+	OnCapture(FloatBufferPtr, InBufferFrames, NumChannels, StreamTime, bOverflow);
 }
 
 bool FAudioCaptureImpl::GetDefaultCaptureDeviceInfo(FCaptureDeviceInfo& OutInfo)
@@ -104,6 +102,8 @@ bool FAudioCaptureImpl::OpenDefaultCaptureStream(FOnCaptureFunction InOnCapture,
 	FCoreDelegates::ApplicationHasEnteredForegroundDelegate.AddRaw(this, &FAudioCaptureImpl::OnApplicationResume);
 
 	// set up variables to be populated by ML Audio
+	uint32 NumFrames = NumFramesDesired;
+	OnCapture = MoveTemp(InOnCapture);
 	uint32 ChannelCount = NumChannels;
 	MLAudioBufferFormat DefaultBufferFormat;
 	uint32 RecommendedBufferSize = 0;
@@ -112,7 +112,7 @@ bool FAudioCaptureImpl::OpenDefaultCaptureStream(FOnCaptureFunction InOnCapture,
 	uint32 NumFrames = NumFramesDesired;
 	uint32 BufferSize = NumFramesDesired * NumChannels * sizeof(int16);
 	// MERGE-REVIEW - check logic here
-	uint32 RequestedBufferSize = StreamParams.NumFramesDesired * NumChannels * sizeof(int16);
+	uint32 RequestedBufferSize = NumFrames * NumChannels * sizeof(int16);
 
 	MLResult Result = MLAudioGetInputStreamDefaults(ChannelCount, UnsignedSampleRate, &DefaultBufferFormat, &RecommendedBufferSize, &MinBufferSize);
 	if (Result != MLResult_Ok)
@@ -123,8 +123,8 @@ bool FAudioCaptureImpl::OpenDefaultCaptureStream(FOnCaptureFunction InOnCapture,
 
 	if (BufferSize < RecommendedBufferSize)
 	{
-		BufferSize = StreamParams.NumFramesDesired * NumChannels;
-		UE_LOG(LogAudioCapture, Display, TEXT("Using buffer size of %u"), StreamParams.NumFramesDesired * NumChannels);
+		BufferSize = NumFrames * NumChannels;
+		UE_LOG(LogAudioCapture, Display, TEXT("Using buffer size of %u"), NumFrames * NumChannels);
 	}
 	UE_LOG(LogAudioCapture, Display, TEXT("Using buffer size of %u"), BufferSize);
 
