Attempted fix to build error part 2

diff --git a/Engine/Plugins/Runtime/AudioCapture/Source/AudioCapture/Private/AudioCaptureMLAudio.cpp b/Engine/Plugins/Runtime/AudioCapture/Source/AudioCapture/Private/AudioCaptureMLAudio.cpp
--- a/Engine/Plugins/Runtime/AudioCapture/Source/AudioCapture/Private/AudioCaptureMLAudio.cpp
+++ b/Engine/Plugins/Runtime/AudioCapture/Source/AudioCapture/Private/AudioCaptureMLAudio.cpp
@@ -102,17 +102,15 @@ bool FAudioCaptureImpl::OpenDefaultCaptureStream(FOnCaptureFunction InOnCapture,
 	FCoreDelegates::ApplicationHasEnteredForegroundDelegate.AddRaw(this, &FAudioCaptureImpl::OnApplicationResume);
 
 	// set up variables to be populated by ML Audio
-	uint32 NumFrames = NumFramesDesired;
 	OnCapture = MoveTemp(InOnCapture);
 	uint32 ChannelCount = NumChannels;
 	MLAudioBufferFormat DefaultBufferFormat;
 	uint32 RecommendedBufferSize = 0;
 	uint32 MinBufferSize = 0;
 	uint32 UnsignedSampleRate = SampleRate;
-	uint32 NumFrames = NumFramesDesired;
 	uint32 BufferSize = NumFramesDesired * NumChannels * sizeof(int16);
 	// MERGE-REVIEW - check logic here
-	uint32 RequestedBufferSize = NumFrames * NumChannels * sizeof(int16);
+	uint32 RequestedBufferSize = NumFramesDesired * NumChannels * sizeof(int16);
 
 	MLResult Result = MLAudioGetInputStreamDefaults(ChannelCount, UnsignedSampleRate, &DefaultBufferFormat, &RecommendedBufferSize, &MinBufferSize);
 	if (Result != MLResult_Ok)
@@ -123,8 +121,8 @@ bool FAudioCaptureImpl::OpenDefaultCaptureStream(FOnCaptureFunction InOnCapture,
 
 	if (BufferSize < RecommendedBufferSize)
 	{
-		BufferSize = NumFrames * NumChannels;
-		UE_LOG(LogAudioCapture, Display, TEXT("Using buffer size of %u"), NumFrames * NumChannels);
+		BufferSize = NumFramesDesired * NumChannels;
+		UE_LOG(LogAudioCapture, Display, TEXT("Using buffer size of %u"), NumFramesDesired * NumChannels);
 	}
 	UE_LOG(LogAudioCapture, Display, TEXT("Using buffer size of %u"), BufferSize);
 
