Fix bad merge in AudioCaptureMLAudio

diff --git a/Engine/Plugins/Runtime/AudioCapture/Source/AudioCapture/Private/AudioCaptureMLAudio.cpp b/Engine/Plugins/Runtime/AudioCapture/Source/AudioCapture/Private/AudioCaptureMLAudio.cpp
--- a/Engine/Plugins/Runtime/AudioCapture/Source/AudioCapture/Private/AudioCaptureMLAudio.cpp
+++ b/Engine/Plugins/Runtime/AudioCapture/Source/AudioCapture/Private/AudioCaptureMLAudio.cpp
@@ -109,8 +109,6 @@ bool FAudioCaptureImpl::OpenDefaultCaptureStream(FOnCaptureFunction InOnCapture,
 	uint32 MinBufferSize = 0;
 	uint32 UnsignedSampleRate = SampleRate;
 	uint32 BufferSize = NumFramesDesired * NumChannels * sizeof(int16);
-	// MERGE-REVIEW - check logic here
-	uint32 RequestedBufferSize = NumFramesDesired * NumChannels * sizeof(int16);
 
 	MLResult Result = MLAudioGetInputStreamDefaults(ChannelCount, UnsignedSampleRate, &DefaultBufferFormat, &RecommendedBufferSize, &MinBufferSize);
 	if (Result != MLResult_Ok)
@@ -119,11 +117,6 @@ bool FAudioCaptureImpl::OpenDefaultCaptureStream(FOnCaptureFunction InOnCapture,
 		return false;
 	}
 
-	if (BufferSize < RecommendedBufferSize)
-	{
-		BufferSize = NumFramesDesired * NumChannels;
-		UE_LOG(LogAudioCapture, Display, TEXT("Using buffer size of %u"), NumFramesDesired * NumChannels);
-	}
 	UE_LOG(LogAudioCapture, Display, TEXT("Using buffer size of %u"), BufferSize);
 
 	DefaultBufferFormat.bits_per_sample = 16;
@@ -131,8 +124,6 @@ bool FAudioCaptureImpl::OpenDefaultCaptureStream(FOnCaptureFunction InOnCapture,
 	DefaultBufferFormat.channel_count = ChannelCount;
 	DefaultBufferFormat.samples_per_second = SampleRate;
 
-	OnCapture = MoveTemp(InOnCapture);
-
 	// Open up new audio stream
 	Result = MLAudioCreateInputFromVoiceComm(&DefaultBufferFormat, BufferSize, &OnAudioCaptureCallback, this, &InputDeviceHandle);
 	if (Result != MLResult_Ok)
