[NO NON-PII COMMIT MESSAGE]

diff --git a/Engine/Plugins/Runtime/AudioCapture/AudioCapture.uplugin b/Engine/Plugins/Runtime/AudioCapture/AudioCapture.uplugin
--- a/Engine/Plugins/Runtime/AudioCapture/AudioCapture.uplugin
+++ b/Engine/Plugins/Runtime/AudioCapture/AudioCapture.uplugin
@@ -20,7 +20,7 @@
 			"Name" : "AudioCapture",
 			"Type" : "Runtime",
       "LoadingPhase": "Default",
- 			"WhitelistPlatforms" : [ "Win64", "Lumin" ]
+ 			"WhitelistPlatforms" : [ "Win64", "Lumin", "Mac" ]
 		},
 		{
 			"Name" : "AudioCaptureEditor",
