Fix for audio capture implementation not being loaded.
Explicitly call LoadModule.
ethan.geller, cristina.riveron

diff --git a/Engine/Plugins/Runtime/AudioCapture/Source/AudioCapture/AudioCapture.Build.cs b/Engine/Plugins/Runtime/AudioCapture/Source/AudioCapture/AudioCapture.Build.cs
--- a/Engine/Plugins/Runtime/AudioCapture/Source/AudioCapture/AudioCapture.Build.cs
+++ b/Engine/Plugins/Runtime/AudioCapture/Source/AudioCapture/AudioCapture.Build.cs
@@ -18,6 +18,7 @@ namespace UnrealBuildTool.Rules
 
             if (Target.Platform == UnrealTargetPlatform.Win32 ||
                 Target.Platform == UnrealTargetPlatform.Win64 ||
+				Target.Platform == UnrealTargetPlatform.Mac   ||
                 Target.Platform == UnrealTargetPlatform.XboxOne)
             {
                 PrivateDependencyModuleNames.Add("AudioCaptureRtAudio");
diff --git a/Engine/Plugins/Runtime/AudioCapture/Source/AudioCapture/Private/AudioCapture.cpp b/Engine/Plugins/Runtime/AudioCapture/Source/AudioCapture/Private/AudioCapture.cpp
--- a/Engine/Plugins/Runtime/AudioCapture/Source/AudioCapture/Private/AudioCapture.cpp
+++ b/Engine/Plugins/Runtime/AudioCapture/Source/AudioCapture/Private/AudioCapture.cpp
@@ -92,6 +92,18 @@ UAudioCapture* UAudioCaptureFunctionLibrary::CreateAudioCapture()
 
 void FAudioCaptureModule::StartupModule()
 {
+	// Load platform specific implementations for audio capture:
+#if PLATFORM_WINDOWS || PLATFORM_MAC || PLATFORM_XBOXONE
+	FModuleManager::Get().LoadModule(TEXT("AudioCaptureRtAudio"));
+#elif PLATFORM_IOS
+	FModuleManager::Get().LoadModule(TEXT("AudioCaptureAudioUnit"));
+#elif PLATFORM_PS4
+	FModuleManager::Get().LoadModule(TEXT("AudioCapturePS4Voice"));
+#elif PLATFORM_ANDROID
+	FModuleManager::Get().LoadModule(TEXT("AudioCaptureAndroid"));
+#elif PLATFORM_SWITCH
+	FModuleManager::Get().LoadModule(TEXT("AudioCaptureSwitch"));
+#endif
 }
 
 void FAudioCaptureModule::ShutdownModule()
