fix for usoundwave created with mic capture in the Take Recorder window not playing back.

diff --git a/Engine/Plugins/Runtime/AudioCapture/Source/AudioCaptureEditor/Private/AudioRecordingManager.cpp b/Engine/Plugins/Runtime/AudioCapture/Source/AudioCaptureEditor/Private/AudioRecordingManager.cpp
--- a/Engine/Plugins/Runtime/AudioCapture/Source/AudioCaptureEditor/Private/AudioRecordingManager.cpp
+++ b/Engine/Plugins/Runtime/AudioCapture/Source/AudioCaptureEditor/Private/AudioRecordingManager.cpp
@@ -405,7 +405,7 @@ void FAudioRecordingManager::StopRecording(TArray<USoundWave*>& OutSoundWaves)
 					}
 
 					// Compressed data is now out of date.
-					NewSoundWave->InvalidateCompressedData();
+					NewSoundWave->InvalidateCompressedData(true, false);
 
 					// Copy the raw wave data file to the sound wave for storage. Will allow the recording to be exported.
 					NewSoundWave->RawData.Lock(LOCK_READ_WRITE);
