- Fixed non-unity compile

diff --git a/Engine/Plugins/Runtime/AudioCapture/Source/AudioCapture/Private/AudioCapture.cpp b/Engine/Plugins/Runtime/AudioCapture/Source/AudioCapture/Private/AudioCapture.cpp
--- a/Engine/Plugins/Runtime/AudioCapture/Source/AudioCapture/Private/AudioCapture.cpp
+++ b/Engine/Plugins/Runtime/AudioCapture/Source/AudioCapture/Private/AudioCapture.cpp
@@ -5,6 +5,7 @@
 #include "Misc/ScopeLock.h"
 #include "Modules/ModuleManager.h"
 #include "AudioCaptureCore.h"
+#include "Misc/ConfigCacheIni.h"
 
 DEFINE_LOG_CATEGORY(LogAudioCapture);
 
