Fix for compiler warning in initialization of FAudioInputDeviceInfo

diff --git a/Engine/Plugins/Runtime/AudioCapture/Source/AudioCapture/Public/AudioCaptureBlueprintLibrary.h b/Engine/Plugins/Runtime/AudioCapture/Source/AudioCapture/Public/AudioCaptureBlueprintLibrary.h
--- a/Engine/Plugins/Runtime/AudioCapture/Source/AudioCapture/Public/AudioCaptureBlueprintLibrary.h
+++ b/Engine/Plugins/Runtime/AudioCapture/Source/AudioCapture/Public/AudioCaptureBlueprintLibrary.h
@@ -25,7 +25,14 @@ struct AUDIOCAPTURE_API FAudioInputDeviceInfo
 {
 	GENERATED_USTRUCT_BODY()
 
-	FAudioInputDeviceInfo() {};
+	FAudioInputDeviceInfo() :
+		DeviceName(""),
+		DeviceId(""),
+		InputChannels(0),
+		PreferredSampleRate(0),
+		bSupportsHardwareAEC(true) 
+	{};
+
 	FAudioInputDeviceInfo(const Audio::FCaptureDeviceInfo& InDeviceInfo);
 
 	/** The name of the audio device */
