Allowing procedural sound waves (and synth components) have more than 2 channels

diff --git a/Engine/Plugins/Runtime/AudioCapture/Source/AudioCapture/Private/AudioCaptureComponent.cpp b/Engine/Plugins/Runtime/AudioCapture/Source/AudioCapture/Private/AudioCaptureComponent.cpp
--- a/Engine/Plugins/Runtime/AudioCapture/Source/AudioCapture/Private/AudioCaptureComponent.cpp
+++ b/Engine/Plugins/Runtime/AudioCapture/Source/AudioCapture/Private/AudioCaptureComponent.cpp
@@ -23,8 +23,7 @@ bool UAudioCaptureComponent::Init(int32& SampleRate)
 		SampleRate = DeviceInfo.PreferredSampleRate;
 		NumChannels = DeviceInfo.InputChannels;
 
-		// Only support mono and stereo mic inputs for now...
-		if (NumChannels == 1 || NumChannels == 2)
+		if (NumChannels > 0 && NumChannels <= 8)
 		{
 			// This may fail if capture synths aren't supported on a given platform or if something went wrong with the capture device
 			bIsStreamOpen = CaptureSynth.OpenDefaultStream();
@@ -32,7 +31,7 @@ bool UAudioCaptureComponent::Init(int32& SampleRate)
 		}
 		else
 		{
-			UE_LOG(LogAudio, Warning, TEXT("Audio capture components only support mono and stereo mic input."));
+			UE_LOG(LogAudio, Warning, TEXT("Attempted to open a device with the invalid NumChannels value of %i"), NumChannels);
 		}
 
 	}
diff --git a/Engine/Source/Runtime/AudioCaptureImplementations/AudioCaptureRtAudio/Private/AudioCaptureRtAudio.cpp b/Engine/Source/Runtime/AudioCaptureImplementations/AudioCaptureRtAudio/Private/AudioCaptureRtAudio.cpp
--- a/Engine/Source/Runtime/AudioCaptureImplementations/AudioCaptureRtAudio/Private/AudioCaptureRtAudio.cpp
+++ b/Engine/Source/Runtime/AudioCaptureImplementations/AudioCaptureRtAudio/Private/AudioCaptureRtAudio.cpp
@@ -81,7 +81,7 @@ bool Audio::FAudioCaptureRtAudioStream::OpenCaptureStream(const FAudioCaptureDev
 	RtAudio::StreamParameters RtAudioStreamParams;
 	RtAudioStreamParams.deviceId = InputDeviceId;
 	RtAudioStreamParams.firstChannel = 0;
-	RtAudioStreamParams.nChannels = FMath::Min((int32)DeviceInfo.inputChannels, 2);
+	RtAudioStreamParams.nChannels = FMath::Min((int32)DeviceInfo.inputChannels, 8);
 
 	if (CaptureDevice.isStreamOpen())
 	{
diff --git a/Engine/Source/Runtime/AudioMixer/Private/Components/SynthComponent.cpp b/Engine/Source/Runtime/AudioMixer/Private/Components/SynthComponent.cpp
--- a/Engine/Source/Runtime/AudioMixer/Private/Components/SynthComponent.cpp
+++ b/Engine/Source/Runtime/AudioMixer/Private/Components/SynthComponent.cpp
@@ -223,12 +223,12 @@ void USynthComponent::Initialize(int32 SampleRateOverride)
 		// Initialize the synth component
 		Init(SampleRate);
 
-		if (NumChannels < 0 || NumChannels > 2)
+		if (NumChannels < 0 || NumChannels > 8)
 		{
-			UE_LOG(LogAudioMixer, Error, TEXT("Synthesis component '%s' has set an invalid channel count '%d' (only mono and stereo currently supported)."), *GetName(), NumChannels);
+			UE_LOG(LogAudioMixer, Error, TEXT("Synthesis component '%s' has set an invalid channel count '%d'."), *GetName(), NumChannels);
 		}
 
-		NumChannels = FMath::Clamp(NumChannels, 1, 2);
+		NumChannels = FMath::Clamp(NumChannels, 1, 8);
 #endif
 
 		if (!Synth)
