Fix crash in AudioCapture when no audio devices are enabled

diff --git a/Engine/Plugins/Runtime/AudioCapture/Source/AudioCapture/Private/AudioCaptureComponent.cpp b/Engine/Plugins/Runtime/AudioCapture/Source/AudioCapture/Private/AudioCaptureComponent.cpp
--- a/Engine/Plugins/Runtime/AudioCapture/Source/AudioCapture/Private/AudioCaptureComponent.cpp
+++ b/Engine/Plugins/Runtime/AudioCapture/Source/AudioCapture/Private/AudioCaptureComponent.cpp
@@ -23,7 +23,14 @@ bool UAudioCaptureComponent::Init(int32& SampleRate)
 	Audio::FCaptureDeviceInfo DeviceInfo;
 	if (CaptureSynth.GetDefaultCaptureDeviceInfo(DeviceInfo))
 	{
-		SampleRate = DeviceInfo.PreferredSampleRate;
+		if (DeviceInfo.PreferredSampleRate > 0)
+		{
+			SampleRate = DeviceInfo.PreferredSampleRate;
+		}
+		else
+		{
+			UE_LOG(LogAudio, Warning, TEXT("Attempted to open a capture device with the invalid SampleRate value of %i"), DeviceInfo.PreferredSampleRate);
+		}
 		NumChannels = DeviceInfo.InputChannels;
 
 		if (NumChannels > 0 && NumChannels <= 8)
