Fix for audio capture component crashing on level transitions

diff --git a/Engine/Plugins/Runtime/AudioCapture/Source/AudioCapture/Private/AudioCaptureComponent.cpp b/Engine/Plugins/Runtime/AudioCapture/Source/AudioCapture/Private/AudioCaptureComponent.cpp
--- a/Engine/Plugins/Runtime/AudioCapture/Source/AudioCapture/Private/AudioCaptureComponent.cpp
+++ b/Engine/Plugins/Runtime/AudioCapture/Source/AudioCapture/Private/AudioCaptureComponent.cpp
@@ -66,6 +66,9 @@ void UAudioCaptureComponent::BeginDestroy()
 
 bool UAudioCaptureComponent::IsReadyForFinishDestroy()
 {
+	//Prevent an infinite loop here if it was marked pending kill while generating
+	OnEndGenerate();
+
 	return !bIsNotReadyForForFinishDestroy;
 }
 
