Enable up to eight audio input sources in Take Recorder.
Added audio input device list UI to the project settings which enables the user to select a specific device to use when recording audio.

diff --git a/Engine/Plugins/Runtime/AudioCapture/Source/AudioCaptureEditor/AudioCaptureEditor.build.cs b/Engine/Plugins/Runtime/AudioCapture/Source/AudioCaptureEditor/AudioCaptureEditor.build.cs
--- a/Engine/Plugins/Runtime/AudioCapture/Source/AudioCaptureEditor/AudioCaptureEditor.build.cs
+++ b/Engine/Plugins/Runtime/AudioCapture/Source/AudioCaptureEditor/AudioCaptureEditor.build.cs
@@ -13,7 +13,6 @@ public class AudioCaptureEditor : ModuleRules
 				"Engine",
 				"EditorFramework",
 				"UnrealEd",
-				"SequenceRecorder",
 				"AudioEditor",
 				"AudioCapture",
 				"AudioCaptureCore",
diff --git a/Engine/Plugins/Runtime/AudioCapture/Source/AudioCaptureEditor/Private/AudioCaptureEditor.cpp b/Engine/Plugins/Runtime/AudioCapture/Source/AudioCaptureEditor/Private/AudioCaptureEditor.cpp
--- a/Engine/Plugins/Runtime/AudioCapture/Source/AudioCaptureEditor/Private/AudioCaptureEditor.cpp
+++ b/Engine/Plugins/Runtime/AudioCapture/Source/AudioCaptureEditor/Private/AudioCaptureEditor.cpp
@@ -1,59 +1,93 @@
 // Copyright Epic Games, Inc. All Rights Reserved.
 
 #include "AudioCaptureEditor.h"
-#include "ISequenceRecorder.h"
-#include "ISequenceAudioRecorder.h"
-#include "Modules/ModuleManager.h"
 #include "AudioRecordingManager.h"
+#include "IAudioCaptureEditorModule.h"
+#include "Modules/ModuleManager.h"
+
 
-class FSequenceAudioRecorder : public ISequenceAudioRecorder
+namespace Audio
 {
-public:
+	bool FAudioCaptureEditor::IsReadyToRecord()
+	{
+		return AudioRecordingManager.IsReadyToRecord();
+	}
 
-	virtual void Start(const FSequenceAudioRecorderSettings& InSettings) override
+	bool FAudioCaptureEditor::IsRecording()
 	{
-		Audio::FRecordingSettings RecordingSettings;
-		RecordingSettings.AssetName = InSettings.AssetName;
-		RecordingSettings.Directory = InSettings.Directory;
-		RecordingSettings.GainDb = InSettings.GainDb;
-		RecordingSettings.RecordingDuration = InSettings.RecordingDuration;
-		RecordingSettings.bSplitChannels = InSettings.bSplitChannels;
+		return AudioRecordingManager.IsRecording();
+	}
 
-		TArray<USoundWave*> OutSoundWaves;
-		Audio::FAudioRecordingManager::Get().StartRecording(RecordingSettings, OutSoundWaves);
+	bool FAudioCaptureEditor::IsStopped()
+	{
+		return AudioRecordingManager.IsStopped();
 	}
 
-	virtual void Stop(TArray<USoundWave*>& OutSoundWaves) override
+	void FAudioCaptureEditor::Start(const FTakeRecorderAudioSettings& InSettings)
 	{
-		Audio::FAudioRecordingManager::Get().StopRecording(OutSoundWaves);
+		FRecordingSettings RecordSettings;
+		RecordSettings.AudioCaptureDeviceId = InSettings.AudioCaptureDeviceId;
+		RecordSettings.AudioInputBufferSize = InSettings.AudioInputBufferSize;
+		RecordSettings.RecordingDuration = InSettings.RecordingDuration;
+		RecordSettings.NumRecordChannels = InSettings.NumRecordChannels;
+
+		AudioRecordingManager.StartRecording(RecordSettings);
 	}
-};
 
+	void FAudioCaptureEditor::Stop()
+	{
+		AudioRecordingManager.StopRecording();
+	}
 
-class FAudioCaptureEditorModule : public IModuleInterface
-{
-	virtual void StartupModule() override
+	TObjectPtr<USoundWave> FAudioCaptureEditor::GetRecordedSoundWave(const FTakeRecorderAudioSourceSettings& InSourceSettings)
 	{
-		ISequenceRecorder& Recorder = FModuleManager::Get().LoadModuleChecked<ISequenceRecorder>("SequenceRecorder");
-		RecorderHandle = Recorder.RegisterAudioRecorder(
-			[](){
-				return TUniquePtr<ISequenceAudioRecorder>(new FSequenceAudioRecorder());
-			}
-		);
+		FRecordingManagerSourceSettings RecorderSettings;
+		RecorderSettings.AssetName = InSourceSettings.AssetName;
+		RecorderSettings.Directory = InSourceSettings.Directory;
+		RecorderSettings.GainDb = InSourceSettings.GainDb;
+		RecorderSettings.InputChannelNumber = InSourceSettings.InputChannelNumber;
+		RecorderSettings.StartTimecode = InSourceSettings.StartTimecode;
+		RecorderSettings.VideoFrameRate = InSourceSettings.VideoFrameRate;
+
+		return AudioRecordingManager.GetRecordedSoundWave(RecorderSettings);
 	}
 
-	virtual void ShutdownModule() override
+	bool FAudioCaptureEditor::GetCaptureDeviceInfo(FTakeRecorderAudioDeviceInfo& OutInfo, int32 InDeviceIndex)
 	{
-		ISequenceRecorder* Recorder = FModuleManager::Get().GetModulePtr<ISequenceRecorder>("SequenceRecorder");
-		
-		if (Recorder)
+		FCaptureDeviceInfo DeviceInfo;
+		bool bFoundDevice = AudioRecordingManager.GetCaptureDeviceInfo(DeviceInfo, InDeviceIndex);
+		if (bFoundDevice)
 		{
-			Recorder->UnregisterAudioRecorder(RecorderHandle);
+			CopyDeviceInfo(DeviceInfo, OutInfo);
 		}
+
+		return bFoundDevice;
 	}
 
-	FDelegateHandle RecorderHandle;
-};
+	bool FAudioCaptureEditor::GetCaptureDevicesAvailable(TArray<FTakeRecorderAudioDeviceInfo>& OutDevices)
+	{
+		TArray<FCaptureDeviceInfo> DeviceArray;
+		bool bFoundDevice = AudioRecordingManager.GetCaptureDevicesAvailable(DeviceArray);
+		if (bFoundDevice)
+		{
+			for (const FCaptureDeviceInfo& DeviceInfo : DeviceArray)
+			{
+				FTakeRecorderAudioDeviceInfo TempInfo;
+				CopyDeviceInfo(DeviceInfo, TempInfo);
 
+				OutDevices.Add(MoveTemp(TempInfo));
+			}
+		}
+
+		return bFoundDevice;
+	}
+
+	void FAudioCaptureEditor::CopyDeviceInfo(const FCaptureDeviceInfo& InDeviceInfo, FTakeRecorderAudioDeviceInfo& OutDeviceInfo)
+	{
+		OutDeviceInfo.DeviceId = InDeviceInfo.DeviceId;
+		OutDeviceInfo.DeviceName = InDeviceInfo.DeviceName;
+		OutDeviceInfo.InputChannels = InDeviceInfo.InputChannels;
+		OutDeviceInfo.PreferredSampleRate = InDeviceInfo.PreferredSampleRate;
+	}
 
-IMPLEMENT_MODULE( FAudioCaptureEditorModule, AudioCaptureEditor );
+}
diff --git a/Engine/Plugins/Runtime/AudioCapture/Source/AudioCaptureEditor/Private/AudioCaptureEditor.h b/Engine/Plugins/Runtime/AudioCapture/Source/AudioCaptureEditor/Private/AudioCaptureEditor.h
--- a/Engine/Plugins/Runtime/AudioCapture/Source/AudioCaptureEditor/Private/AudioCaptureEditor.h
+++ b/Engine/Plugins/Runtime/AudioCapture/Source/AudioCaptureEditor/Private/AudioCaptureEditor.h
@@ -2,5 +2,35 @@
 
 #pragma once
 
-#include "CoreMinimal.h"
+#include "IAudioCaptureEditor.h"
 
+#include "AudioCaptureDeviceInterface.h"
+#include "AudioRecordingManager.h"
+
+namespace Audio
+{
+	class FAudioCaptureEditor : public IAudioCaptureEditor
+	{
+	public:
+
+		virtual ~FAudioCaptureEditor() {}
+
+		// Begin IAudioCaptureEditor overrides
+		virtual bool IsReadyToRecord() override;
+		virtual bool IsRecording() override;
+		virtual bool IsStopped() override;
+		virtual void Start(const FTakeRecorderAudioSettings& InSettings) override;
+		virtual void Stop() override;
+		virtual TObjectPtr<USoundWave> GetRecordedSoundWave(const FTakeRecorderAudioSourceSettings& InSourceSettings) override;
+		virtual bool GetCaptureDeviceInfo(FTakeRecorderAudioDeviceInfo& OutInfo, int32 InDeviceIndex) override;
+		virtual bool GetCaptureDevicesAvailable(TArray<FTakeRecorderAudioDeviceInfo>& OutDevices) override;
+		// End IAudioCaptureEditor overrides
+
+	private:
+
+		void CopyDeviceInfo(const Audio::FCaptureDeviceInfo& InDeviceInfo, FTakeRecorderAudioDeviceInfo& OutDeviceInfo);
+
+		Audio::FAudioRecordingManager AudioRecordingManager;
+	};
+
+}
diff --git a/Engine/Plugins/Runtime/AudioCapture/Source/AudioCaptureEditor/Private/AudioCaptureEditorModule.cpp b/Engine/Plugins/Runtime/AudioCapture/Source/AudioCaptureEditor/Private/AudioCaptureEditorModule.cpp
new file mode 100644
--- /dev/null
+++ b/Engine/Plugins/Runtime/AudioCapture/Source/AudioCaptureEditor/Private/AudioCaptureEditorModule.cpp
@@ -0,0 +1,60 @@
+// Copyright Epic Games, Inc. All Rights Reserved.
+
+#include "AudioCaptureEditor.h"
+#include "IAudioCaptureEditorModule.h"
+#include "Modules/ModuleManager.h"
+
+
+class FAudioCaptureEditorModule : public IAudioCaptureEditorModule
+{
+private:
+
+	virtual void StartupModule() override
+	{
+		FactoryDelegate = RegisterAudioRecorder(
+			[]() {
+				return TUniquePtr<IAudioCaptureEditor>(new Audio::FAudioCaptureEditor());
+			}
+		);
+	}
+
+	virtual void ShutdownModule() override
+	{
+		UnregisterAudioRecorder(FactoryDelegate);
+	}
+
+	virtual FDelegateHandle RegisterAudioRecorder(const TFunction<TUniquePtr<IAudioCaptureEditor>()>& FactoryFunction) override
+	{
+		ensureMsgf(!AudioFactory, TEXT("Audio recorder already registered."));
+
+		AudioFactory = FactoryFunction;
+		AudioFactoryHandle = FDelegateHandle(FDelegateHandle::GenerateNewHandle);
+		return AudioFactoryHandle;
+	}
+
+	virtual void UnregisterAudioRecorder(FDelegateHandle Handle) override
+	{
+		if (Handle == AudioFactoryHandle)
+		{
+			AudioFactory = nullptr;
+			AudioFactoryHandle = FDelegateHandle();
+		}
+	}
+
+	virtual bool HasAudioRecorder() const override
+	{
+		return AudioFactoryHandle.IsValid();
+	}
+
+	virtual TUniquePtr<IAudioCaptureEditor> CreateAudioRecorder() const override
+	{
+		return AudioFactory ? AudioFactory() : TUniquePtr<IAudioCaptureEditor>();
+	}
+
+	TFunction<TUniquePtr<IAudioCaptureEditor>()> AudioFactory;
+	FDelegateHandle AudioFactoryHandle;
+
+	FDelegateHandle FactoryDelegate;
+};
+
+IMPLEMENT_MODULE( FAudioCaptureEditorModule, AudioCaptureEditor );
diff --git a/Engine/Plugins/Runtime/AudioCapture/Source/AudioCaptureEditor/Private/AudioRecordingManager.cpp b/Engine/Plugins/Runtime/AudioCapture/Source/AudioCaptureEditor/Private/AudioRecordingManager.cpp
--- a/Engine/Plugins/Runtime/AudioCapture/Source/AudioCaptureEditor/Private/AudioRecordingManager.cpp
+++ b/Engine/Plugins/Runtime/AudioCapture/Source/AudioCaptureEditor/Private/AudioRecordingManager.cpp
@@ -6,6 +6,7 @@
 #include "Sound/AudioSettings.h"
 #include "AssetRegistry/AssetRegistryModule.h"
 #include "Factories/SoundFactory.h"
+#include "Misc/EngineVersion.h"
 #include "Misc/ScopeLock.h"
 #include "Sound/SoundWave.h"
 #include "AudioDeviceManager.h"
@@ -25,14 +26,12 @@ namespace Audio
 */
 
 FAudioRecordingManager::FAudioRecordingManager()
-	: NumRecordedSamples(0)
+	: RecorderState(EAudioRecorderState::PreRecord)
+	, NumRecordedSamples(0)
 	, NumFramesToRecord()
 	, RecordingBlockSize(0)
 	, RecordingSampleRate(44100.0f)
 	, NumInputChannels(1)
-	, InputGain(1.0f)
-	, bRecording(false)
-	, bError(false)
 {
 }
 
@@ -44,111 +43,98 @@ FAudioRecordingManager::~FAudioRecordingManager()
 	}
 }
 
-FAudioRecordingManager& FAudioRecordingManager::Get()
+void FAudioRecordingManager::StartRecording(const FRecordingSettings& InSettings)
 {
-	// Static isngleton instance
-	static FAudioRecordingManager MicInputManager;
-	return MicInputManager;
-}
+	EAudioRecorderState CurrentState = RecorderState.load();
+	check(CurrentState == EAudioRecorderState::PreRecord);
 
-void FAudioRecordingManager::StartRecording(const FRecordingSettings& InSettings, TArray<USoundWave*>& OutSoundWaves)
-{
-	if (bError)
+	if (CurrentState == EAudioRecorderState::PreRecord)
 	{
-		return;
-	}
-
-	// Stop any recordings currently going on (if there is one)
-	StopRecording(OutSoundWaves);
-
-	// Default to 1024 blocks
-	RecordingBlockSize = 1024;
+		RecorderState = EAudioRecorderState::Recording;
 
-	// If we have a stream open close it (reusing streams can cause a blip of previous recordings audio)
-	if (AudioCapture.IsStreamOpen())
-	{
-		AudioCapture.StopStream();
-		AudioCapture.CloseStream();
-	}
+		// Default to 1024 blocks
+		RecordingBlockSize = InSettings.AudioInputBufferSize;
 
-	UE_LOG(LogMicManager, Log, TEXT("Starting mic recording."));
+		// If we have a stream open close it (reusing streams can cause a blip of previous recordings audio)
+		if (AudioCapture.IsStreamOpen())
+		{
+			AudioCapture.StopStream();
+			AudioCapture.CloseStream();
+		}
 
-	// Copy the settings
-	Settings = InSettings;
+		UE_LOG(LogMicManager, Log, TEXT("Starting mic recording."));
 
-	// Convert input gain decibels into linear scale
-	InputGain = Audio::ConvertToLinear(Settings.GainDb);
+		// Copy the settings
+		Settings = InSettings;
 
-	// Get the default mic input device info
-	FAudioCaptureDeviceParams StreamParams;
-	FCaptureDeviceInfo DeviceInfo;
-	if (!AudioCapture.GetCaptureDeviceInfo(DeviceInfo, StreamParams.DeviceIndex))
-	{
-		UE_LOG(LogMicManager, Error, TEXT("Unable to get default audio capture device."));
-		return;
-	}
-
-	RecordingSampleRate = DeviceInfo.PreferredSampleRate;
-	NumInputChannels = DeviceInfo.InputChannels;
-
-	if (NumInputChannels == 0)
-	{
-		UE_LOG(LogMicManager, Warning, TEXT("No audio channels to record."));
-		return;
-	}
+		FCaptureDeviceInfo DeviceInfo;
+		int32 DeviceIndex = INDEX_NONE;
+		if (!GetCaptureDeviceInfo(DeviceInfo, DeviceIndex, Settings.AudioCaptureDeviceId))
+		{
+			UE_LOG(LogMicManager, Error, TEXT("Unable to get audio capture device."));
+			return;
+		}
 
-	// Reserve enough space in our current recording buffer for 10 seconds of audio to prevent slowing down due to allocations
-	if (Settings.RecordingDuration > 0.0f)
-	{
-		int32 NumSamplesToReserve = Settings.RecordingDuration * RecordingSampleRate * NumInputChannels;
-		CurrentRecordedPCMData.Reset(NumSamplesToReserve);
+		RecordingSampleRate = DeviceInfo.PreferredSampleRate;
+		NumInputChannels = FMath::Min(DeviceInfo.InputChannels, Settings.NumRecordChannels);
 
-		// Set a limit on the number of frames to record
-		NumFramesToRecord = RecordingSampleRate * Settings.RecordingDuration;
-	}
-	else
-	{
-		// if not given a duration, resize array to 60 seconds of audio anyway
-		int32 NumSamplesToReserve = 60 * RecordingSampleRate * NumInputChannels;
-		CurrentRecordedPCMData.Reset(NumSamplesToReserve);
+		if (NumInputChannels == 0)
+		{
+			UE_LOG(LogMicManager, Warning, TEXT("No audio channels to record."));
+			return;
+		}
 
-		// Set to be INDEX_NONE otherwise
-		NumFramesToRecord = INDEX_NONE;
-	}
+		// Initialize Sound Waves array
+		RecordedSoundWaves.Init(nullptr, DeviceInfo.InputChannels);
 
-	// If we have more than 2 channels, we're going to force splitting up the assets since we don't propertly support multi-channel USoundWave assets
-	if (NumInputChannels > 2)
-	{
-		Settings.bSplitChannels = true;
-	}
+		// Reserve enough space in our current recording buffer for 10 seconds of audio to prevent slowing down due to allocations
+		if (Settings.RecordingDuration > 0.0f)
+		{
+			int32 NumSamplesToReserve = Settings.RecordingDuration * RecordingSampleRate * NumInputChannels;
+			CurrentRecordedPCMData.Reset(NumSamplesToReserve);
 
-	NumRecordedSamples = 0;
-	NumOverflowsDetected = 0;
+			// Set a limit on the number of frames to record
+			NumFramesToRecord = RecordingSampleRate * Settings.RecordingDuration;
+		}
+		else
+		{
+			// if not given a duration, resize array to 60 seconds of audio anyway
+			int32 NumSamplesToReserve = 60 * RecordingSampleRate * NumInputChannels;
+			CurrentRecordedPCMData.Reset(NumSamplesToReserve);
 
-	// Publish to the mic input thread that we're ready to record...
-	bRecording = true;
+			// Set to be INDEX_NONE otherwise
+			NumFramesToRecord = INDEX_NONE;
+		}
 
-	StreamParams.NumInputChannels = NumInputChannels;
-	StreamParams.SampleRate = RecordingSampleRate;
-	StreamParams.PCMAudioEncoding = EPCMAudioEncoding::PCM_16;
+		NumRecordedSamples = 0;
+		NumOverflowsDetected = 0;
 
-	uint32 BufferFrames = FMath::Max(RecordingBlockSize, 256);
+		FAudioCaptureDeviceParams StreamParams;
+		StreamParams.NumInputChannels = NumInputChannels;
+		StreamParams.SampleRate = RecordingSampleRate;
+		StreamParams.PCMAudioEncoding = EPCMAudioEncoding::PCM_16;
+		StreamParams.DeviceIndex = DeviceIndex;
 
-	UE_LOG(LogMicManager, Log, TEXT("Initialized mic recording manager at %.4f hz sample rate, %d channels, and %d Recording Block Size"), RecordingSampleRate, StreamParams.NumInputChannels, BufferFrames);
+		uint32 BufferFrames = FMath::Max(RecordingBlockSize, 256);
 
-	FOnAudioCaptureFunction OnAudioCaptureCallback = [this](const void* InBuffer, int32 InNumFrames, int32 InNumChannels, int32 InSampleRate, double InStreamTime, bool bInOverflow)
-	{
-		OnAudioCapture(InBuffer, InNumFrames, InStreamTime, bInOverflow);
-	};
+		UE_LOG(LogMicManager, Log, TEXT("Initialized mic recording manager at %.4f hz sample rate, %d channels, and %d Recording Block Size"), RecordingSampleRate, StreamParams.NumInputChannels, BufferFrames);
 
-	// Open up new audio stream
-	bool bSuccess = AudioCapture.OpenAudioCaptureStream(StreamParams, OnAudioCaptureCallback, BufferFrames);
+		FOnAudioCaptureFunction OnAudioCaptureCallback = [this](const void* InBuffer, int32 InNumFrames, int32 InNumChannels, int32 InSampleRate, double InStreamTime, bool bInOverflow)
+		{
+			OnAudioCapture(InBuffer, InNumFrames, InStreamTime, bInOverflow);
+		};
 
-	bError = !bSuccess;
+		// Open up new audio stream
+		bool bSuccess = AudioCapture.OpenAudioCaptureStream(StreamParams, OnAudioCaptureCallback, BufferFrames);
 
-	if (bSuccess)
-	{
-		AudioCapture.StartStream();
+		if (bSuccess)
+		{
+			AudioCapture.StartStream();
+		}
+		else
+		{
+			RecorderState = EAudioRecorderState::ErrorDetected;
+		}
 	}
 }
 
@@ -199,14 +185,15 @@ static void SampleRateConvert(float CurrentSR, float TargetSR, int32 NumChannels
 	}
 }
 
-void FAudioRecordingManager::StopRecording(TArray<USoundWave*>& OutSoundWaves)
+void FAudioRecordingManager::StopRecording()
 {
-	FScopeLock Lock(&CriticalSection);
+	EAudioRecorderState CurrentState = RecorderState.load();
+	check(CurrentState == EAudioRecorderState::Recording);
 
-	// If we're currently recording, stop the recording
-	if (bRecording)
+	if (CurrentState == EAudioRecorderState::Recording)
 	{
-		bRecording = false;
+		RecorderState = EAudioRecorderState::Stopped;
+
 		AudioCapture.StopStream();
 		AudioCapture.CloseStream();
 
@@ -221,212 +208,257 @@ void FAudioRecordingManager::StopRecording(TArray<USoundWave*>& OutSoundWaves)
 				NumRecordedSamples = CurrentRecordedPCMData.Num();
 			}
 
-			UE_LOG(LogMicManager, Log, TEXT("Stopping mic recording. Recorded %d frames of audio (%.4f seconds). Detected %d buffer overflows."), 
-				NumRecordedSamples, 
+			UE_LOG(LogMicManager, Log, TEXT("Stopping mic recording. Recorded %d frames of audio (%.4f seconds). Detected %d buffer overflows."),
+				NumRecordedSamples,
 				((float)NumRecordedSamples / RecordingSampleRate) / (float)NumInputChannels,
 				NumOverflowsDetected);
+		}
+	}
+}
 
-			// Get a ptr to the buffer we're actually going to serialize
-			TArray<int16>* PCMDataToSerialize = nullptr;
+TObjectPtr<USoundWave> FAudioRecordingManager::GetRecordedSoundWave(const FRecordingManagerSourceSettings& InSourceSettings)
+{
+	EAudioRecorderState CurrentState = RecorderState.load();
+	check(CurrentState == EAudioRecorderState::Stopped);
 
-			// If our sample rate isn't 44100, then we need to do a SRC
-			if (RecordingSampleRate != 44100.0f)
-			{
-				UE_LOG(LogMicManager, Log, TEXT("Converting sample rate from %d hz to 44100 hz."), (int32)RecordingSampleRate);
+	if (CurrentState == EAudioRecorderState::Stopped)
+	{
+		if (PCMDataToSerialize == nullptr)
+		{
+			ProcessRecordedData(InSourceSettings);
+		}
 
-				SampleRateConvert(RecordingSampleRate, (float)WAVE_FILE_SAMPLERATE, NumInputChannels, CurrentRecordedPCMData, NumRecordedSamples, ConvertedPCMData);
+		if (RecordedSoundWaves[InSourceSettings.InputChannelNumber - 1] != nullptr)
+		{
+			return RecordedSoundWaves[InSourceSettings.InputChannelNumber - 1];
+		}
 
-				// Update the recorded samples to the converted buffer samples
-				NumRecordedSamples = ConvertedPCMData.Num();
-				PCMDataToSerialize = &ConvertedPCMData;
-			}
-			else
-			{
-				// Just use the original recorded buffer to serialize
-				PCMDataToSerialize = &CurrentRecordedPCMData;
-			}
+		return CreateSoundWaveAsset(InSourceSettings);
+	}
 
-			// Scale by the linear gain if it's been set to something
-			if (InputGain != 1.0f)
-			{
-				UE_LOG(LogMicManager, Log, TEXT("Scaling gain of recording by %.2f linear gain."), InputGain);
+	return nullptr;
+}
 
-				for (int32 i = 0; i < NumRecordedSamples; ++i)
-				{
-					// Scale by input gain, clamp to prevent integer overflow when casting back to int16. Will still clip.
-					(*PCMDataToSerialize)[i] = (int16)FMath::Clamp(InputGain * (float)(*PCMDataToSerialize)[i], -32767.0f, 32767.0f);
-				}
-			}
+void FAudioRecordingManager::ProcessRecordedData(const FRecordingManagerSourceSettings& InSourceSettings)
+{
+	// If our sample rate isn't 44100, then we need to do a SRC
+	if (RecordingSampleRate != 44100.0f)
+	{
+		UE_LOG(LogMicManager, Log, TEXT("Converting sample rate from %d hz to 44100 hz."), (int32)RecordingSampleRate);
 
-			uint8* RawData = nullptr;
-			int32 NumBytes = 0;
-			int32 NumChannelsToSerialize = 0;
-			int32 NumWavesToSerialize = 0;
+		SampleRateConvert(RecordingSampleRate, (float)WAVE_FILE_SAMPLERATE, NumInputChannels, CurrentRecordedPCMData, NumRecordedSamples, ConvertedPCMData);
 
-			if (Settings.bSplitChannels)
-			{
-				// We're going to serialize several waves, one for each input channel
-				NumWavesToSerialize = NumInputChannels;
+		// Update the recorded samples to the converted buffer samples
+		NumRecordedSamples = ConvertedPCMData.Num();
+		PCMDataToSerialize = &ConvertedPCMData;
+	}
+	else
+	{
+		// Just use the original recorded buffer to serialize
+		PCMDataToSerialize = &CurrentRecordedPCMData;
+	}
+
+	DeinterleaveRecordedData();
+}
+
+void FAudioRecordingManager::DeinterleaveRecordedData()
+{
+	// De-interleaved buffer size will be the number of frames of audio
+	int32 NumFrames = NumRecordedSamples / NumInputChannels;
+
+	// Reset our deinterleaved audio buffer
+	DeinterleavedAudio.Reset(NumInputChannels);
+
+	// Get ptr to the interleaved buffer for speed in non-release builds
+	int16* InterleavedBufferPtr = PCMDataToSerialize->GetData();
+
+	// For every input channel, create a new buffer
+	for (int32 Channel = 0; Channel < NumInputChannels; ++Channel)
+	{
+		// Prepare a new deinterleaved buffer
+		DeinterleavedAudio.Add(FDeinterleavedAudio());
+		FDeinterleavedAudio& DeinterleavedChannelAudio = DeinterleavedAudio[Channel];
 
-				// We're only going to serialize 1 channel of audio at a time
-				NumChannelsToSerialize = 1;
+		DeinterleavedChannelAudio.PCMData.Reset();
+		DeinterleavedChannelAudio.PCMData.AddUninitialized(NumFrames);
 
-				// De-interleaved buffer size will be the number of frames of audio
-				int32 NumFrames = NumRecordedSamples / NumInputChannels;
+		// Get a ptr to the buffer for speed
+		int16* DeinterleavedBufferPtr = DeinterleavedChannelAudio.PCMData.GetData();
 
-				// This is the num bytes per de-interleaved audio buffer
-				NumBytes = NumFrames * sizeof(int16);
+		// Copy every N channel to the deinterleaved buffer, starting with the current channel
+		int32 CurrentSample = Channel;
+		for (int32 Frame = 0; Frame < NumFrames; ++Frame)
+		{
+			// Simply copy the interleaved value to the deinterleaved value
+			DeinterleavedBufferPtr[Frame] = InterleavedBufferPtr[CurrentSample];
 
-				// Reset our deinterleaved audio buffer
-				DeinterleavedAudio.Reset(NumInputChannels);
+			// Increment the stride according the num channels
+			CurrentSample += NumInputChannels;
+		}
+	}
+}
 
-				// Get ptr to the interleaved buffer for speed in non-release builds
-				int16* InterleavedBufferPtr = PCMDataToSerialize->GetData();
+void FAudioRecordingManager::ApplyGain(const FRecordingManagerSourceSettings& InSourceSettings, const int32 InNumFrames)
+{
+	// Convert input gain decibels into linear scale
+	float InputGain = Audio::ConvertToLinear(InSourceSettings.GainDb);
 
-				// For every input channel, create a new buffer
-				for (int32 Channel = 0; Channel < NumInputChannels; ++Channel)
-				{
-					// Prepare a new deinterleaved buffer
-					DeinterleavedAudio.Add(FDeinterleavedAudio());
-					FDeinterleavedAudio& DeinterleavedChannelAudio = DeinterleavedAudio[Channel];
+	// Scale by the linear gain if it's been set to something
+	if (InputGain != 1.0f)
+	{
+		UE_LOG(LogMicManager, Log, TEXT("Scaling gain of recording by %.2f linear gain."), InputGain);
 
-					DeinterleavedChannelAudio.PCMData.Reset();
-					DeinterleavedChannelAudio.PCMData.AddUninitialized(NumFrames);
+		int32 ChannelIndex = InSourceSettings.InputChannelNumber - 1;
+		int16* RawData = DeinterleavedAudio[ChannelIndex].PCMData.GetData();
 
-					// Get a ptr to the buffer for speed
-					int16* DeinterleavedBufferPtr = DeinterleavedChannelAudio.PCMData.GetData();
+		for (int32 SampleIndex = 0; SampleIndex < InNumFrames; ++SampleIndex)
+		{
+			// Scale by input gain, clamp to prevent integer overflow when casting back to int16. Will still clip.
+			RawData[SampleIndex] = (int16)FMath::Clamp(InputGain * (float)RawData[SampleIndex], -32767.0f, 32767.0f);
+		}
+	}
+}
 
-					// Copy every N channel to the deinterleaved buffer, starting with the current channel
-					int32 CurrentSample = Channel;
-					for (int32 Frame = 0; Frame < NumFrames; ++Frame)
-					{
-						// Simply copy the interleaved value to the deinterleaved value
-						DeinterleavedBufferPtr[Frame] = InterleavedBufferPtr[CurrentSample];
+TObjectPtr<USoundWave> FAudioRecordingManager::CreateSoundWaveAsset(const FRecordingManagerSourceSettings& InSourceSettings)
+{
+	uint8* RawData = nullptr;
+	int32 NumChannelsToSerialize = 1;
+	int32 NumWavesToSerialize = NumInputChannels;
 
-						// Increment the stride according the num channels
-						CurrentSample += NumInputChannels;
-					}
-				}
-			}
-			else
-			{
-				// Only need to serialize one channel
-				NumWavesToSerialize = 1;
-				RawData = (uint8*)PCMDataToSerialize->GetData();
-				NumBytes = NumRecordedSamples * sizeof(int16);
-				NumChannelsToSerialize = NumInputChannels;
-			}
+	TObjectPtr<USoundWave> NewSoundWave = nullptr;
+	USoundWave* ExistingSoundWave = nullptr;
+
+	TArray<UAudioComponent*> ComponentsToRestart;
+	bool bCreatedPackage = false;
+
+	int32 NumBytes = 0;
+
+	// De-interleaved buffer size will be the number of frames of audio
+	int32 NumFrames = NumRecordedSamples / NumInputChannels;
+
+	// This is the num bytes per de-interleaved audio buffer
+	NumBytes = NumFrames * sizeof(int16);
+
+	ApplyGain(InSourceSettings, NumFrames);
+
+	if (InSourceSettings.Directory.Path.IsEmpty() || InSourceSettings.AssetName.IsEmpty())
+	{
+		// Create a new sound wave object from the transient package
+		NewSoundWave = NewObject<USoundWave>(GetTransientPackage(), *InSourceSettings.AssetName);
+	}
+	else
+	{
+		FString PackageName;
+		FString AssetName = InSourceSettings.AssetName;
+		AssetName = FString::Printf(TEXT("%s_Channel_%d"), *AssetName, InSourceSettings.InputChannelNumber);
+		PackageName = InSourceSettings.Directory.Path / AssetName;
+
+		int32 ChannelIndex = InSourceSettings.InputChannelNumber - 1;
+		// Get the raw data from the deinterleaved audio location.
+		RawData = (uint8*)DeinterleavedAudio[ChannelIndex].PCMData.GetData();
+
+		UPackage* Package = CreatePackage(*PackageName);
+
+		// Create a raw .wav file to stuff the raw PCM data in so when we create the sound wave asset it's identical to a normal imported asset
+		check(RawData != nullptr);
+		SerializeWaveFile(RawWaveData, RawData, NumBytes, NumChannelsToSerialize, WAVE_FILE_SAMPLERATE);
+
+		// Check to see if a sound wave already exists at this location
+		ExistingSoundWave = FindObject<USoundWave>(Package, *AssetName);
+
+		// See if it's currently being played right now
+		FAudioDeviceManager* AudioDeviceManager = GEngine->GetAudioDeviceManager();
+		if (AudioDeviceManager && ExistingSoundWave)
+		{
+			AudioDeviceManager->StopSoundsUsingResource(ExistingSoundWave, &ComponentsToRestart);
+		}
+
+		// Create a new sound wave object
+		if (ExistingSoundWave)
+		{
+			NewSoundWave = ExistingSoundWave;
+			NewSoundWave->FreeResources();
+		}
+		else
+		{
+			NewSoundWave = NewObject<USoundWave>(Package, *AssetName, RF_Public | RF_Standalone);
+		}
+
+		// Compressed data is now out of date.
+		NewSoundWave->InvalidateCompressedData(true, false);
+
+		// Copy the raw wave data file to the sound wave for storage. Will allow the recording to be exported.
+		FSharedBuffer UpdatedBuffed = FSharedBuffer::Clone(RawWaveData.GetData(), RawWaveData.Num());
+		NewSoundWave->RawData.UpdatePayload(UpdatedBuffed);
+
+		bCreatedPackage = true;
+	}
+
+	if (NewSoundWave)
+	{
+		// Copy the recorded data to the sound wave so we can quickly preview it
+		NewSoundWave->RawPCMDataSize = NumBytes;
+		NewSoundWave->RawPCMData = (uint8*)FMemory::Malloc(NewSoundWave->RawPCMDataSize);
+		FMemory::Memcpy(NewSoundWave->RawPCMData, RawData, NumBytes);
+
+		// Calculate the duration of the sound wave
+		// Note: We use the NumInputChannels for duration calculation since NumChannelsToSerialize may be 1 channel while NumInputChannels is 2 for the "split stereo" feature.
+		NewSoundWave->Duration = (float)(NumRecordedSamples / NumInputChannels) / WAVE_FILE_SAMPLERATE;
+		NewSoundWave->SetSampleRate(WAVE_FILE_SAMPLERATE);
+		NewSoundWave->NumChannels = NumChannelsToSerialize;
+		NewSoundWave->SetSoundAssetCompressionType(ESoundAssetCompressionType::BinkAudio);
+		NewSoundWave->SetTimecodeInfo(GetTimecodeInfo(InSourceSettings));
+
+		if (bCreatedPackage)
+		{
+			//GEditor->GetEditorSubsystem<UImportSubsystem>()->BroadcastAssetPostImport(this, NewSoundWave);
+
+			FAssetRegistryModule::AssetCreated(NewSoundWave);
+			NewSoundWave->MarkPackageDirty();
 
-			// Loop through the number of waves we're going to serialize
-			for (int32 i = 0; i < NumWavesToSerialize; ++i)
+			// Restart any audio components if they need to be restarted
+			for (int32 ComponentIndex = 0; ComponentIndex < ComponentsToRestart.Num(); ++ComponentIndex)
 			{
-				USoundWave* NewSoundWave = nullptr;
-				USoundWave* ExistingSoundWave = nullptr;
-				TArray<UAudioComponent*> ComponentsToRestart;
-				bool bCreatedPackage = false;
-
-				if (Settings.Directory.Path.IsEmpty() || Settings.AssetName.IsEmpty())
-				{
-					// Create a new sound wave object from the transient package
-					NewSoundWave = NewObject<USoundWave>(GetTransientPackage(), *Settings.AssetName);
-				}
-				else
-				{
-					FString PackageName;
-					FString AssetName = Settings.AssetName;
-					if (NumWavesToSerialize != 1)
-					{
-						AssetName = FString::Printf(TEXT("%s_Channel_%d"), *AssetName, i);
-						PackageName = Settings.Directory.Path / AssetName;
-
-						// Get the raw data from the deinterleaved audio location.
-						RawData = (uint8*)DeinterleavedAudio[i].PCMData.GetData();
-					}
-					else
-					{
-						// Create a new package
-						PackageName = Settings.Directory.Path / AssetName;
-					}
-
-					UPackage* Package = CreatePackage(*PackageName);
-
-					// Create a raw .wav file to stuff the raw PCM data in so when we create the sound wave asset it's identical to a normal imported asset
-					check(RawData != nullptr);
-					SerializeWaveFile(RawWaveData, RawData, NumBytes, NumChannelsToSerialize, WAVE_FILE_SAMPLERATE);
-
-					// Check to see if a sound wave already exists at this location
-					ExistingSoundWave = FindObject<USoundWave>(Package, *AssetName);
-
-					// See if it's currently being played right now
-					FAudioDeviceManager* AudioDeviceManager = GEngine->GetAudioDeviceManager();
-					if (AudioDeviceManager && ExistingSoundWave)
-					{
-						AudioDeviceManager->StopSoundsUsingResource(ExistingSoundWave, &ComponentsToRestart);
-					}
-
-					// Create a new sound wave object
-					if (ExistingSoundWave)
-					{
-						NewSoundWave = ExistingSoundWave;
-						NewSoundWave->FreeResources();
-					}
-					else
-					{
-						NewSoundWave = NewObject<USoundWave>(Package, *AssetName, RF_Public | RF_Standalone);
-					}
-
-					// Compressed data is now out of date.
-					NewSoundWave->InvalidateCompressedData(true, false);
-
-					// Copy the raw wave data file to the sound wave for storage. Will allow the recording to be exported.
-					FSharedBuffer UpdatedBuffed = FSharedBuffer::Clone(RawWaveData.GetData(), RawWaveData.Num());
-					NewSoundWave->RawData.UpdatePayload(UpdatedBuffed);
-
-					bCreatedPackage = true;
-				}
-
-				if (NewSoundWave)
-				{
-					// Copy the recorded data to the sound wave so we can quickly preview it
-					NewSoundWave->RawPCMDataSize = NumBytes;
-					NewSoundWave->RawPCMData = (uint8*)FMemory::Malloc(NewSoundWave->RawPCMDataSize);
-					FMemory::Memcpy(NewSoundWave->RawPCMData, RawData, NumBytes);
-
-					// Calculate the duration of the sound wave
-					// Note: We use the NumInputChannels for duration calculation since NumChannelsToSerialize may be 1 channel while NumInputChannels is 2 for the "split stereo" feature.
-					NewSoundWave->Duration = (float)(NumRecordedSamples / NumInputChannels) / WAVE_FILE_SAMPLERATE;
-					NewSoundWave->SetSampleRate(WAVE_FILE_SAMPLERATE);
-					NewSoundWave->NumChannels = NumChannelsToSerialize;
-					NewSoundWave->SetSoundAssetCompressionType(ESoundAssetCompressionType::BinkAudio);
-
-					if (bCreatedPackage)
-					{
-						//GEditor->GetEditorSubsystem<UImportSubsystem>()->BroadcastAssetPostImport(this, NewSoundWave);
-
-						FAssetRegistryModule::AssetCreated(NewSoundWave);
-						NewSoundWave->MarkPackageDirty();
-
-						// Restart any audio components if they need to be restarted
-						for (int32 ComponentIndex = 0; ComponentIndex < ComponentsToRestart.Num(); ++ComponentIndex)
-						{
-							ComponentsToRestart[ComponentIndex]->Play();
-						}
-					}
-
-					OutSoundWaves.Add(NewSoundWave);
-				}
+				ComponentsToRestart[ComponentIndex]->Play();
 			}
 		}
+
+		RecordedSoundWaves[InSourceSettings.InputChannelNumber - 1] = NewSoundWave;
 	}
+
+	return NewSoundWave;
 }
 
-int32 FAudioRecordingManager::OnAudioCapture(const void* InBuffer, uint32 InBufferFrames, double StreamTime, bool bOverflow)
+FSoundWaveTimecodeInfo FAudioRecordingManager::GetTimecodeInfo(const FRecordingManagerSourceSettings& InSourceSettings)
 {
-	// Check recording state prior to locking to avoid potential race condition when stopping
-	if (bRecording && InBufferFrames > 0)
+	FSoundWaveTimecodeInfo TimecodeInfo;
+
+	double SecondsFromMidnight = InSourceSettings.StartTimecode.ToTimespan(InSourceSettings.VideoFrameRate).GetTotalSeconds();
+	if (!ensureAlwaysMsgf(SecondsFromMidnight > 0, TEXT("Invalid SecondsFromMidnight value: %f, defaulting to 0"), SecondsFromMidnight))
 	{
-		FScopeLock Lock(&CriticalSection);
+		SecondsFromMidnight = 0.0;
+	}
 
+	TimecodeInfo.NumSamplesSinceMidnight = static_cast<uint64>(SecondsFromMidnight * WAVE_FILE_SAMPLERATE);
+	TimecodeInfo.NumSamplesPerSecond = WAVE_FILE_SAMPLERATE;
+	TimecodeInfo.Description = InSourceSettings.AssetName;
+	TimecodeInfo.OriginatorDate = FDateTime::Now().ToString(TEXT("%Y-%m-%d"));
+	TimecodeInfo.OriginatorTime = FDateTime::Now().ToString(TEXT("%H:%M:%S"));
+	// Not using FApp::GetName() here because it gets the executable name on disk which can be modified.
+	// This is similar to what the splash screen does.
+	TimecodeInfo.OriginatorDescription = FString(TEXT("Unreal Editor"));
+	TimecodeInfo.OriginatorReference = FEngineVersion::Current().ToString();
+	TimecodeInfo.TimecodeRate = InSourceSettings.VideoFrameRate;
+	TimecodeInfo.bTimecodeIsDropFrame = InSourceSettings.StartTimecode.bDropFrameFormat;
+
+	// Copy elision
+	return TimecodeInfo;
+}
+
+int32 FAudioRecordingManager::OnAudioCapture(const void* InBuffer, uint32 InBufferFrames, double StreamTime, bool bOverflow)
+{
+	if (RecorderState == EAudioRecorderState::Recording && InBufferFrames > 0)
+	{
 		if (bOverflow)
 		{
 			++NumOverflowsDetected;
@@ -439,4 +471,35 @@ int32 FAudioRecordingManager::OnAudioCapture(const void* InBuffer, uint32 InBuff
 	return 1;
 }
 
+bool FAudioRecordingManager::GetCaptureDeviceInfo(FCaptureDeviceInfo& OutInfo, int32 InDeviceIndex)
+{
+	return AudioCapture.GetCaptureDeviceInfo(OutInfo, InDeviceIndex);
+}
+
+bool FAudioRecordingManager::GetCaptureDeviceInfo(FCaptureDeviceInfo& OutInfo, int32& OutDeviceIndex, FString InDeviceId)
+{
+	TArray<FCaptureDeviceInfo> DeviceInfoArray;
+	if (GetCaptureDevicesAvailable(DeviceInfoArray))
+	{
+		for (int32 DeviceIndex = 0; DeviceIndex < DeviceInfoArray.Num(); ++DeviceIndex)
+		{
+			FCaptureDeviceInfo& DeviceInfo = DeviceInfoArray[DeviceIndex];
+			if (DeviceInfo.DeviceId == InDeviceId)
+			{
+				OutInfo = DeviceInfo;
+				OutDeviceIndex = DeviceIndex;
+
+				return true;
+			}
+		}
+	}
+
+	return false;
+}
+
+bool FAudioRecordingManager::GetCaptureDevicesAvailable(TArray<FCaptureDeviceInfo>& OutDevices)
+{
+	return AudioCapture.GetCaptureDevicesAvailable(OutDevices) != 0;
+}
+
 } // namespace Audio
diff --git a/Engine/Plugins/Runtime/AudioCapture/Source/AudioCaptureEditor/Private/AudioRecordingManager.h b/Engine/Plugins/Runtime/AudioCapture/Source/AudioCaptureEditor/Private/AudioRecordingManager.h
--- a/Engine/Plugins/Runtime/AudioCapture/Source/AudioCaptureEditor/Private/AudioRecordingManager.h
+++ b/Engine/Plugins/Runtime/AudioCapture/Source/AudioCaptureEditor/Private/AudioRecordingManager.h
@@ -5,6 +5,9 @@
 #include "AudioCaptureCore.h"
 #include "Engine/EngineTypes.h"
 #include "HAL/ThreadSafeBool.h"
+#include "Misc/FrameRate.h"
+#include "Misc/Timecode.h"
+#include "Sound/SoundWaveTimecodeInfo.h"
 
 class USoundWave;
 
@@ -14,6 +17,22 @@ namespace Audio
 {
 	// Settings used for starting a recording with the recording manager
 	struct FRecordingSettings
+	{
+		// The audio device to use for capturing
+		FString AudioCaptureDeviceId;
+
+		// The desired buffer size in samples.
+		int32 AudioInputBufferSize = 1024;
+
+		// The length of the recording. Set to 0.0 or less for unbounded recording duration.
+		float RecordingDuration = -1.0f;
+
+		// Number of channels to record 
+		int32 NumRecordChannels = 0;
+	};
+
+	// Settings used for individual audio input sources
+	struct FRecordingManagerSourceSettings
 	{
 		// The place to put the audio file recording
 		FDirectoryPath Directory;
@@ -21,20 +40,17 @@ namespace Audio
 		// The name of the asset you want to record
 		FString AssetName;
 
-		// The length of the recording. Set to 0.0 or less for unbounded recording duration.
-		float RecordingDuration;
-
 		// The output gain of the recording. Allows you to boost the amplitude of your recording.
-		float GainDb;
+		float GainDb = 0.0f;
+
+		// The channel number of the selected audio device to use for recording.
+		int32 InputChannelNumber = -1;
 
-		// Whether or not to split the asset into separate assets for each microphone input channel
-		bool bSplitChannels;
+		// The timecode locations for this recording
+		FTimecode StartTimecode;
 
-		FRecordingSettings()
-			: RecordingDuration(0.0f)
-			, GainDb(0.0f)
-			, bSplitChannels(false)
-		{}
+		// The video framerate used during this recording
+		FFrameRate VideoFrameRate;
 	};
 
 	/**
@@ -44,39 +60,79 @@ namespace Audio
 	class FAudioRecordingManager
 	{
 	public:
-		// Retrieves the singleton recording manager
-		static FAudioRecordingManager& Get();
+		
+		FAudioRecordingManager();
+
+		~FAudioRecordingManager();
+
+		// Returns true if this object is in the PreRecord state indicating it is ready to record
+		bool IsReadyToRecord() { return RecorderState == EAudioRecorderState::PreRecord; }
+
+		// Returns true if this object is in the Recording state indicating it is currently recording
+		bool IsRecording() { return RecorderState == EAudioRecorderState::Recording; }
+
+		// Returns true if this object is in the Stopped state indicating it has finished recording
+		bool IsStopped() { return RecorderState == EAudioRecorderState::Stopped; }
 
 		// Starts a new recording with the given name and optional duration. 
 		// If set to -1.0f, a duration won't be used and the recording length will be determined by StopRecording().
-		void StartRecording(const FRecordingSettings& Settings, TArray<USoundWave*>& OutSoundWaves);
+		void StartRecording(const FRecordingSettings& Settings);
 
 		// Stops recording if the recording manager is recording. If not recording but has recorded data (due to set duration), it will just return the generated USoundWave.
-		void StopRecording(TArray<USoundWave*>& OutSoundWaves);
+		void StopRecording();
 
 		// Called by RtAudio when a new audio buffer is ready to be supplied.
 		int32 OnAudioCapture(const void* InBuffer, uint32 InBufferFrames, double StreamTime, bool bOverflow);
 
+		TObjectPtr<USoundWave> GetRecordedSoundWave(const FRecordingManagerSourceSettings& InSourceSettings);
+
+		// Returns the audio capture device information at the given index.
+		bool GetCaptureDeviceInfo(FCaptureDeviceInfo& OutInfo, int32 InDeviceIndex = INDEX_NONE);
+
+		// Returns the total amount of audio devices.
+		bool GetCaptureDevicesAvailable(TArray<FCaptureDeviceInfo>& OutDevices);
+
 	private:
 
-		// Private Constructor
-		FAudioRecordingManager();
+		// States used by AudioRecordingManager state machine
+		enum class EAudioRecorderState : uint8
+		{
+			PreRecord,
+			Recording,
+			Stopped,
+			ErrorDetected
+		};
 
-		// Private Destructor
-		~FAudioRecordingManager();
+		// Called after recording to handle processing of recorded data (SRC, de-interleaving, apply gain, etc.)
+		void ProcessRecordedData(const FRecordingManagerSourceSettings& InSourceSettings);
+
+		// Deinterleave recorded data into discreet channels
+		void DeinterleaveRecordedData();
+
+		// Applies the user supplied gain parameter if neccessary to a single channel of audio
+		void ApplyGain(const FRecordingManagerSourceSettings& InSourceSettings, const int32 InNumFrames);
+
+		// Creates a wave asset for a single channel of audio data
+		TObjectPtr<USoundWave> CreateSoundWaveAsset(const FRecordingManagerSourceSettings& InSourceSettings);
+
+		// Returns the device info and index for a given device id
+		bool GetCaptureDeviceInfo(FCaptureDeviceInfo& OutInfo, int32& OutDeviceIndex, FString InDeviceId);
+
+		// Returns the timecode info for this recording
+		FSoundWaveTimecodeInfo GetTimecodeInfo(const FRecordingManagerSourceSettings& InSourceSettings);
 
 		// The default mic sample rate
 		const int32 WAVE_FILE_SAMPLERATE = 44100;
 
+		// Current state for the state machine
+		std::atomic<EAudioRecorderState> RecorderState;
+
 		// Copy of the input recording settings
 		FRecordingSettings Settings;
 
 		// FAudioCapture object -- used to interact with low-level audio device.
 		FAudioCapture AudioCapture;
 
-		// Critical section used to stop and retrieve finished audio buffers.
-		FCriticalSection CriticalSection;
-
 		// The data which is currently being recorded to, if the manager is actively recording. This is not safe to access while recording.
 		TArray<int16> CurrentRecordedPCMData;
 
@@ -108,18 +164,14 @@ namespace Audio
 		// Num input channels
 		int32 NumInputChannels;
 
-		// A linear gain to apply on mic input
-		float InputGain;
-
-		// Whether or not the manager is actively recording.
-		FThreadSafeBool bRecording;
-
 		// Number of overflows detected while recording
 		int32 NumOverflowsDetected;
 
-		// Whether or not we have an error
-		uint32 bError : 1;
+		// Pointer to the recorded PCM data after any sample rate conversion has been applied
+		TArray<int16>* PCMDataToSerialize = nullptr;
 
+		// Array of USoundWaves created during this record pass
+		TArray<TObjectPtr<USoundWave>> RecordedSoundWaves;
 	};
 
 }
diff --git a/Engine/Plugins/Runtime/AudioCapture/Source/AudioCaptureEditor/Public/AudioCaptureEditorTypes.h b/Engine/Plugins/Runtime/AudioCapture/Source/AudioCaptureEditor/Public/AudioCaptureEditorTypes.h
new file mode 100644
--- /dev/null
+++ b/Engine/Plugins/Runtime/AudioCapture/Source/AudioCaptureEditor/Public/AudioCaptureEditorTypes.h
@@ -0,0 +1,58 @@
+// Copyright Epic Games, Inc. All Rights Reserved.
+
+#pragma once
+
+#include "Engine/EngineTypes.h"
+#include "Misc/FrameRate.h"
+#include "Misc/Timecode.h"
+
+/** Sequencer version of the audio device info. */
+struct FTakeRecorderAudioDeviceInfo
+{
+	/** The name for this audio device*/
+	FString DeviceName;
+	/** The unique identifier for this device */
+	FString DeviceId;
+	/** The number of input channels this device is capable of capturing */
+	int32 InputChannels = 0;
+	/** The preferred sample rate for this device */
+	int32 PreferredSampleRate = 0;
+};
+
+/** Take Recorder version of the audio recording settings. */
+struct FTakeRecorderAudioSettings
+{
+	/** Audio device to use for capture. */
+	FString AudioCaptureDeviceId;
+
+	/** The desired buffer size in samples. */
+	int32 AudioInputBufferSize = 1024;
+
+	/** Optional audio recording duration in seconds. */
+	float RecordingDuration = -1.0f;
+
+	/** Number of channels to record */
+	int32 NumRecordChannels = 0;
+};
+
+/** Take Recorder version of the audio source settings. */
+struct FTakeRecorderAudioSourceSettings
+{
+	/** Directory to create the asset within (empty for transient package) */
+	FDirectoryPath Directory;
+
+	/** Name of the asset. */
+	FString AssetName;
+
+	/** Gain in decibels of the output asset. */
+	float GainDb = 0.0f;
+
+	/** The channel number of the selected audio device to use for recording. */
+	int32 InputChannelNumber = -1;
+
+	/** The timecode position for this recording */
+	FTimecode StartTimecode;
+
+	/** The video frame rate used during this recording */
+	FFrameRate VideoFrameRate;
+};
diff --git a/Engine/Plugins/Runtime/AudioCapture/Source/AudioCaptureEditor/Public/IAudioCaptureEditor.h b/Engine/Plugins/Runtime/AudioCapture/Source/AudioCaptureEditor/Public/IAudioCaptureEditor.h
new file mode 100644
--- /dev/null
+++ b/Engine/Plugins/Runtime/AudioCapture/Source/AudioCaptureEditor/Public/IAudioCaptureEditor.h
@@ -0,0 +1,41 @@
+// Copyright Epic Games, Inc. All Rights Reserved.
+
+#pragma once
+
+#include "AudioCaptureEditorTypes.h"
+#include "Misc/FrameRate.h"
+#include "Misc/Timecode.h"
+
+class USoundWave;
+
+class IAudioCaptureEditor
+{
+public:
+
+	/** Virtual destructor */
+	virtual ~IAudioCaptureEditor() {}
+
+	/** Returns true if this object is in the PreRecord state indicating it is ready to record */
+	virtual bool IsReadyToRecord() = 0;
+
+	/** Returns true if this object is in the Recording state indicating it is currently recording */
+	virtual bool IsRecording() = 0;
+
+	/** Returns true if this object is in the Stopped state indicating it has finished recording */
+	virtual bool IsStopped() = 0;
+
+	/** Start recording audio data */
+	virtual void Start(const FTakeRecorderAudioSettings& InSettings) = 0;
+
+	/** Stop recording audio data */
+	virtual void Stop() = 0;
+
+	/** Returns previously recorded audio data for single channel as a USoundWave */
+	virtual TObjectPtr<USoundWave> GetRecordedSoundWave(const FTakeRecorderAudioSourceSettings& InSourceSettings) = 0;
+
+	/** Returns the audio capture device information at the given index */
+	virtual bool GetCaptureDeviceInfo(FTakeRecorderAudioDeviceInfo& OutInfo, int32 DeviceIndex = INDEX_NONE) = 0;
+
+	/** Returns the total amount of audio devices */
+	virtual bool GetCaptureDevicesAvailable(TArray<FTakeRecorderAudioDeviceInfo>& OutDevices) = 0;
+};
diff --git a/Engine/Plugins/Runtime/AudioCapture/Source/AudioCaptureEditor/Public/IAudioCaptureEditorModule.h b/Engine/Plugins/Runtime/AudioCapture/Source/AudioCaptureEditor/Public/IAudioCaptureEditorModule.h
new file mode 100644
--- /dev/null
+++ b/Engine/Plugins/Runtime/AudioCapture/Source/AudioCaptureEditor/Public/IAudioCaptureEditorModule.h
@@ -0,0 +1,42 @@
+// Copyright Epic Games, Inc. All Rights Reserved.
+
+#pragma once
+
+#include "IAudioCaptureEditor.h"
+#include "Delegates/IDelegateInstance.h"
+#include "Modules/ModuleInterface.h"
+#include "Templates/Function.h"
+#include "Templates/UniquePtr.h"
+
+class ITakeRecorderAudio;
+
+class IAudioCaptureEditorModule : public IModuleInterface
+{
+public:
+
+	/**
+	 * Register a function that will return a new audio capturer for the specified parameters
+	 * @param	FactoryFunction		Function used to generate a new audio recorder
+	 * @return A handle to be passed to UnregisterAudioRecorder to unregister the recorder
+	 */
+	virtual FDelegateHandle RegisterAudioRecorder(const TFunction<TUniquePtr<IAudioCaptureEditor>()>& FactoryFunction) = 0;
+
+	/**
+	 * Unregister a previously registered audio recorder factory function
+	 * @param	RegisteredHandle	The handle returned from RegisterAudioRecorder
+	 */
+	virtual void UnregisterAudioRecorder(FDelegateHandle RegisteredHandle) = 0;
+
+	/**
+	 * Check whether we have an audio recorder registered or not
+	 * @return true if we have an audio recorder registered, false otherwise
+	 */
+	virtual bool HasAudioRecorder() const = 0;
+
+	/**
+	 * Attempt to create an audio recorder
+	 * @param	Settings	Settings for the audio recorder
+	 * @return A valid ptr to an audio recorder or null
+	 */
+	virtual TUniquePtr<IAudioCaptureEditor> CreateAudioRecorder() const = 0;
+};
diff --git a/Engine/Plugins/VirtualProduction/Takes/Source/TakeRecorder/Private/TakeRecorderProjectSettingsCustomization.h b/Engine/Plugins/VirtualProduction/Takes/Source/TakeRecorder/Private/TakeRecorderProjectSettingsCustomization.h
--- a/Engine/Plugins/VirtualProduction/Takes/Source/TakeRecorder/Private/TakeRecorderProjectSettingsCustomization.h
+++ b/Engine/Plugins/VirtualProduction/Takes/Source/TakeRecorder/Private/TakeRecorderProjectSettingsCustomization.h
@@ -2,15 +2,32 @@
 
 #pragma once
 
+#include "IDetailChildrenBuilder.h"
 #include "IDetailCustomization.h"
 #include "DetailLayoutBuilder.h"
 #include "DetailCategoryBuilder.h"
+#include "DetailWidgetRow.h"
+#include "Modules/ModuleManager.h"
 #include "TakeRecorderSettings.h"
+#include "TakeRecorderSourceProperty.h"
+#include "Widgets/TakeRecorderAudioSettingsCustomization.h"
 
 class FTakeRecorderProjectSettingsCustomization : public IDetailCustomization
 {
 	virtual void CustomizeDetails(IDetailLayoutBuilder& DetailLayout) override
 	{
+		// Enumerate audio devices before building the UI. Note, this can be expensive depending on the hardware
+		// attached to the machine, so we wait as late as possible before enumerating.
+		if (UTakeRecorderAudioInputSettings* AudioInputSettings = TakeRecorderAudioSettingsUtils::GetTakeRecorderAudioInputSettings())
+		{
+			AudioInputSettings->EnumerateAudioDevices();
+		}
+
+		FPropertyEditorModule& PropertyEditorModule = FModuleManager::Get().LoadModuleChecked<FPropertyEditorModule>("PropertyEditor");
+
+		// Register audio device details customization for the project settings window
+		PropertyEditorModule.RegisterCustomPropertyTypeLayout(FName(TEXT("AudioInputDeviceProperty")), FOnGetPropertyTypeCustomizationInstance::CreateLambda(&MakeShared<FAudioInputDevicePropertyCustomization>));
+
 		// Pop the take recorder category to the top
 		DetailLayout.EditCategory("Take Recorder");
 
diff --git a/Engine/Plugins/VirtualProduction/Takes/Source/TakeRecorder/Private/Widgets/SLevelSequenceTakeEditor.cpp b/Engine/Plugins/VirtualProduction/Takes/Source/TakeRecorder/Private/Widgets/SLevelSequenceTakeEditor.cpp
--- a/Engine/Plugins/VirtualProduction/Takes/Source/TakeRecorder/Private/Widgets/SLevelSequenceTakeEditor.cpp
+++ b/Engine/Plugins/VirtualProduction/Takes/Source/TakeRecorder/Private/Widgets/SLevelSequenceTakeEditor.cpp
@@ -2,6 +2,7 @@
 
 #include "Widgets/SLevelSequenceTakeEditor.h"
 #include "Widgets/TakeRecorderWidgetConstants.h"
+#include "Widgets/TakeRecorderAudioSettingsCustomization.h"
 #include "TakePreset.h"
 #include "TakeMetaData.h"
 #include "TakeRecorderSource.h"
@@ -10,6 +11,7 @@
 #include "TakeRecorderStyle.h"
 #include "Widgets/STakeRecorderSources.h"
 #include "TakeRecorderModule.h"
+#include "TakeRecorderSettings.h"
 
 #include "LevelSequence.h"
 
@@ -37,9 +39,12 @@
 #include "Widgets/Text/STextBlock.h"
 #include "Widgets/Input/SButton.h"
 #include "Widgets/Input/SComboButton.h"
+#include "Widgets/Input/SComboBox.h"
+#include "Widgets/Input/SEditableTextBox.h"
 #include "Widgets/Layout/SSplitter.h"
 #include "Widgets/Layout/SScrollBox.h"
 #include "Widgets/Layout/SBorder.h"
+#include "Widgets/Views/SListView.h"
 #include "Framework/MultiBox/MultiBoxBuilder.h"
 #include "Styling/SlateIconFinder.h"
 #include "SPositiveActionButton.h"
@@ -60,6 +65,7 @@
 #include "IDetailChildrenBuilder.h"
 #include "IDetailCustomization.h"
 #include "IDetailRootObjectCustomization.h"
+#include "IPropertyUtilities.h"
 #include "DetailWidgetRow.h"
 
 // Widget Includes
@@ -686,6 +692,133 @@ class FRecorderPropertyMapCustomization : public IPropertyTypeCustomization
 };
 const FString FRecorderPropertyMapCustomization::PropertyPathDelimiter = FString(TEXT("."));
 
+class FAudioInputChannelPropertyCustomization : public IPropertyTypeCustomization
+{
+public:
+
+	virtual void CustomizeHeader(TSharedRef<IPropertyHandle> PropertyHandle, FDetailWidgetRow& HeaderRow, IPropertyTypeCustomizationUtils& ) {}
+
+	virtual void CustomizeChildren(TSharedRef<IPropertyHandle> PropertyHandle, IDetailChildrenBuilder& ChildBuilder, IPropertyTypeCustomizationUtils& CustomizationUtils)
+	{
+		InputDeviceChannelHandle = PropertyHandle->GetChildHandle(GET_MEMBER_NAME_CHECKED(FAudioInputDeviceChannelProperty, AudioInputDeviceChannel));
+
+		if (InputDeviceChannelHandle.IsValid())
+		{
+			IDetailPropertyRow& InputDeviceChannelPropertyRow = ChildBuilder.AddProperty(PropertyHandle);
+			InputDeviceChannelPropertyRow.CustomWidget()
+			.NameContent()
+			[
+				InputDeviceChannelHandle->CreatePropertyNameWidget()
+			]
+			.ValueContent()
+			[
+				MakeInputChannelSelectorWidget()
+			];
+		}
+
+		// Register delegate so UI can update when the audio input device changes
+		if (UTakeRecorderAudioInputSettings* AudioInputSettings = TakeRecorderAudioSettingsUtils::GetTakeRecorderAudioInputSettings())
+		{
+			AudioInputSettings->GetOnAudioInputDeviceChanged().Add(FSimpleDelegate::CreateSP(this, &FAudioInputChannelPropertyCustomization::RebuildInputChannelArray));
+		}
+	}
+
+private:
+
+	void BuildInputChannelArray()
+	{
+		if (UTakeRecorderAudioInputSettings* AudioInputSettings = TakeRecorderAudioSettingsUtils::GetTakeRecorderAudioInputSettings())
+		{
+			int32 DeviceChannelCount = AudioInputSettings->GetDeviceChannelCount();
+
+			if (DeviceChannelCount > 0)
+			{
+				for (int32 ChannelIndex = 1; ChannelIndex <= DeviceChannelCount; ++ChannelIndex)
+				{
+					TSharedPtr<int32> ChannelIndexPtr = MakeShared<int32>(ChannelIndex);
+					InputDeviceChannelArray.Add(ChannelIndexPtr);
+				}
+			}
+
+			int32 CurrentValue;
+			InputDeviceChannelHandle->GetValue(CurrentValue);
+			if (CurrentValue > DeviceChannelCount)
+			{
+				InputDeviceChannelHandle->SetValue(0);
+			}
+		}
+	}
+
+	void RebuildInputChannelArray()
+	{
+		InputDeviceChannelArray.Empty();
+		BuildInputChannelArray();
+		ChannelComboBox->RefreshOptions();
+
+		UTakeRecorderAudioInputSettings* AudioInputSettings = TakeRecorderAudioSettingsUtils::GetTakeRecorderAudioInputSettings();
+
+		if (AudioInputSettings && InputDeviceChannelHandle.IsValid())
+		{
+			int32 CurrentValue;
+			InputDeviceChannelHandle->GetValue(CurrentValue);
+			int32 DeviceChannelCount = AudioInputSettings->GetDeviceChannelCount();
+
+			if (CurrentValue < 1 || CurrentValue > DeviceChannelCount)
+			{
+				InputChannelTitleBlock->SetText(FText());
+			}
+		}
+	}
+
+	TSharedRef<SWidget> MakeInputChannelSelectorWidget()
+	{
+		BuildInputChannelArray();
+
+		// Combo box component:
+		SAssignNew(ChannelComboBox, SComboBox<TSharedPtr<int32>>)
+			.OptionsSource(&InputDeviceChannelArray)
+			.OnGenerateWidget_Lambda([](TSharedPtr<int32> InValue)
+			{
+				return SNew(STextBlock)
+					.Text(FText::AsNumber(*InValue))
+					.Font(FAppStyle::Get().GetFontStyle("SmallFont"));
+			})
+			.OnSelectionChanged_Lambda([this](TSharedPtr<int32> InSelection, ESelectInfo::Type InSelectInfo)
+			{
+				if (InSelection.IsValid())
+				{
+					InputDeviceChannelHandle->SetValue(*InSelection);
+					InputChannelTitleBlock->SetText(FText::AsNumber(*InSelection));
+				}
+			})
+			.OnComboBoxOpening_Lambda([this]()
+			{
+			})
+			[
+				SAssignNew(InputChannelTitleBlock, STextBlock)
+				.Text_Lambda([this]()
+				{
+					int32 ChannelNumber;
+					InputDeviceChannelHandle->GetValue(ChannelNumber);
+
+					if (ChannelNumber > 0)
+					{
+						return FText::AsNumber(ChannelNumber);
+					}
+					return FText();
+				})
+				.Font(FAppStyle::Get().GetFontStyle("SmallFont"))
+			];
+
+		return ChannelComboBox.ToSharedRef();
+	}
+
+	TSharedPtr<IPropertyHandle> InputDeviceChannelHandle;
+	TArray<TSharedPtr<int32>> InputDeviceChannelArray;
+	TSharedPtr<STextBlock> InputChannelTitleBlock;
+	TSharedPtr<SComboBox<TSharedPtr<int32>>> ChannelComboBox;
+};
+
 class FRecorderSourceObjectCustomization : public IDetailCustomization
 {
 	virtual void CustomizeDetails(IDetailLayoutBuilder& DetailBuilder) override
@@ -763,6 +896,7 @@ void SLevelSequenceTakeEditor::AddDetails(const TPair<const UClass*, TArray<UObj
 		// We may want to change this in future but it seems like the neatest way to make top level categories have helpful names.
 		Details->RegisterInstancedCustomPropertyLayout(UTakeRecorderSource::StaticClass(), FOnGetDetailCustomizationInstance::CreateLambda(&MakeShared<FRecorderSourceObjectCustomization>));
 
+		Details->RegisterInstancedCustomPropertyTypeLayout(FName(TEXT("AudioInputDeviceChannelProperty")), FOnGetPropertyTypeCustomizationInstance::CreateLambda(&MakeShared<FAudioInputChannelPropertyCustomization >));
 		Details->RegisterInstancedCustomPropertyTypeLayout(FName(TEXT("ActorRecorderPropertyMap")), FOnGetPropertyTypeCustomizationInstance::CreateLambda(&MakeShared<FRecorderPropertyMapCustomization >));
 		Details->RegisterInstancedCustomPropertyTypeLayout(FName(TEXT("ActorRecordedProperty")), FOnGetPropertyTypeCustomizationInstance::CreateLambda(&MakeShared<FRecordedPropertyCustomization >));
 		Details->SetObjects(Pair.Value);
diff --git a/Engine/Plugins/VirtualProduction/Takes/Source/TakeRecorder/Private/Widgets/TakeRecorderAudioSettingsCustomization.cpp b/Engine/Plugins/VirtualProduction/Takes/Source/TakeRecorder/Private/Widgets/TakeRecorderAudioSettingsCustomization.cpp
new file mode 100644
--- /dev/null
+++ b/Engine/Plugins/VirtualProduction/Takes/Source/TakeRecorder/Private/Widgets/TakeRecorderAudioSettingsCustomization.cpp
@@ -0,0 +1,287 @@
+// Copyright Epic Games, Inc. All Rights Reserved.
+
+#pragma once
+
+#include "TakeRecorderAudioSettingsCustomization.h"
+#include "IDetailChildrenBuilder.h"
+#include "TakeRecorderSettings.h"
+#include "TakeRecorderSourceProperty.h"
+#include "Widgets/Input/SCheckBox.h"
+#include "Widgets/Input/SComboBox.h"
+
+void FAudioInputDevicePropertyCustomization::CustomizeChildren(TSharedRef<IPropertyHandle> PropertyHandle, IDetailChildrenBuilder& ChildBuilder, IPropertyTypeCustomizationUtils& CustomizationUtils)
+{
+	UseSystemDefaultHandle = PropertyHandle->GetChildHandle(GET_MEMBER_NAME_CHECKED(FAudioInputDeviceProperty, bUseSystemDefaultAudioDevice));
+	InputDeviceHandle = PropertyHandle->GetChildHandle(GET_MEMBER_NAME_CHECKED(FAudioInputDeviceProperty, DeviceId));
+	DeviceChannelCountHandle = PropertyHandle->GetChildHandle(GET_MEMBER_NAME_CHECKED(FAudioInputDeviceProperty, DeviceInputChannelCount));
+	BufferSizeHandle = PropertyHandle->GetChildHandle(GET_MEMBER_NAME_CHECKED(FAudioInputDeviceProperty, AudioInputBufferSize));
+	DeviceInfoArrayHandle = PropertyHandle->GetChildHandle(GET_MEMBER_NAME_CHECKED(FAudioInputDeviceProperty, DeviceInfoArray))->AsArray();
+
+	// Create check box which indicates system default audio device should be used (disables device combobox)
+	if (UseSystemDefaultHandle.IsValid())
+	{
+		IDetailPropertyRow& UseSystemDefaultRow = ChildBuilder.AddProperty(PropertyHandle);
+		UseSystemDefaultRow.CustomWidget()
+		.NameContent()
+		[
+			UseSystemDefaultHandle->CreatePropertyNameWidget()
+		]
+		.ValueContent()
+		[
+			SNew(SCheckBox)
+			.IsChecked_Lambda([&]() -> ECheckBoxState
+			{
+				bool bUseSystemDefault;
+				UseSystemDefaultHandle->GetValue(bUseSystemDefault);
+
+				return bUseSystemDefault ? ECheckBoxState::Checked : ECheckBoxState::Unchecked;
+			})
+			.OnCheckStateChanged_Lambda([&](ECheckBoxState NewState) -> void
+			{
+				bool bUseSystemDefault = (NewState == ECheckBoxState::Checked);
+				UseSystemDefaultHandle->SetValue(bUseSystemDefault);
+				UpdateDeviceCombobxEnabled();
+			})
+		];
+	}
+
+	// Create audio input device combobox
+	if (InputDeviceHandle.IsValid() && DeviceChannelCountHandle.IsValid() && DeviceInfoArrayHandle.IsValid())
+	{
+		IDetailPropertyRow& InputDevicePropertyRow = ChildBuilder.AddProperty(PropertyHandle);
+
+		InputDevicePropertyRow.CustomWidget()
+		.NameContent()
+		[
+			InputDeviceHandle->CreatePropertyNameWidget()
+		]
+		.ValueContent()
+		[
+			MakeAudioInputSelectorWidget()
+		];
+	}
+
+	// Create audio callback buffer size widget
+	if (BufferSizeHandle.IsValid())
+	{
+		IDetailPropertyRow& BufferSizePropertyRow = ChildBuilder.AddProperty(PropertyHandle);
+		BufferSizePropertyRow.CustomWidget()
+		.NameContent()
+		[
+			BufferSizeHandle->CreatePropertyNameWidget()
+		]
+		.ValueContent()
+		[
+			BufferSizeHandle->CreatePropertyValueWidget()
+		];
+	}
+
+	// Register delegate so UI can update when the audio input device changes
+	// This is needed because the device menu lives in both the project settings window and the take recorder panel
+	if (UTakeRecorderAudioInputSettings* AudioInputSettings = TakeRecorderAudioSettingsUtils::GetTakeRecorderAudioInputSettings())
+	{
+		AudioInputSettings->GetOnAudioInputDeviceChanged().Add(FSimpleDelegate::CreateSP(this, &FAudioInputDevicePropertyCustomization::SynchronizeWidgetStates));
+	}
+}
+
+FString FAudioInputDevicePropertyCustomization::GetDeviceNameFromInfoProperty(const TSharedPtr<IPropertyHandle>& InDeviceInfoHandle)
+{
+	FString DeviceNameStr;
+	if (InDeviceInfoHandle && InDeviceInfoHandle.IsValid())
+	{
+		TSharedPtr<IPropertyHandle> DeviceNameHandle = InDeviceInfoHandle->GetChildHandle(GET_MEMBER_NAME_CHECKED(FAudioInputDeviceInfoProperty, DeviceName));
+		if (DeviceNameHandle.IsValid())
+		{
+			DeviceNameHandle->GetValue(DeviceNameStr);
+		}
+	}
+
+	return DeviceNameStr;
+}
+
+FString FAudioInputDevicePropertyCustomization::GetDeviceIdFromInfoProperty(const TSharedPtr<IPropertyHandle>& InDeviceInfoHandle)
+{
+	FString DeviceIdStr;
+	if (InDeviceInfoHandle.IsValid())
+	{
+		TSharedPtr<IPropertyHandle> DeviceIdHandle = InDeviceInfoHandle->GetChildHandle(GET_MEMBER_NAME_CHECKED(FAudioInputDeviceInfoProperty, DeviceId));
+		if (DeviceIdHandle.IsValid())
+		{
+			DeviceIdHandle->GetValue(DeviceIdStr);
+		}
+	}
+
+	return DeviceIdStr;
+}
+
+int32 FAudioInputDevicePropertyCustomization::GetChannelCountFromInfoProperty(const TSharedPtr<IPropertyHandle>& InDeviceInfoHandle)
+{
+	int32 ChannelCount = 0;
+	if (InDeviceInfoHandle.IsValid())
+	{
+		TSharedPtr<IPropertyHandle> ChannelCountHandle = InDeviceInfoHandle->GetChildHandle(GET_MEMBER_NAME_CHECKED(FAudioInputDeviceInfoProperty, InputChannels));
+		if (ChannelCountHandle.IsValid())
+		{
+			ChannelCountHandle->GetValue(ChannelCount);
+		}
+	}
+
+	return ChannelCount;
+}
+
+bool FAudioInputDevicePropertyCustomization::GetIsDefaultFromInfoProperty(const TSharedPtr<IPropertyHandle>& InDeviceInfoHandle)
+{
+	bool bIsDefaultDevice = false;
+	if (InDeviceInfoHandle.IsValid())
+	{
+		TSharedPtr<IPropertyHandle> IsDefaultHandle = InDeviceInfoHandle->GetChildHandle(GET_MEMBER_NAME_CHECKED(FAudioInputDeviceInfoProperty, bIsDefaultDevice));
+		if (IsDefaultHandle.IsValid())
+		{
+			IsDefaultHandle->GetValue(bIsDefaultDevice);
+		}
+	}
+
+	return bIsDefaultDevice;
+}
+
+TSharedRef<SWidget> FAudioInputDevicePropertyCustomization::MakeAudioInputSelectorWidget()
+{
+	FString CurrentDeviceId;
+	InputDeviceHandle->GetValue(CurrentDeviceId);
+
+	bool bHaveCurrentDevice = CurrentDeviceId.Len() != 0;
+	TSharedPtr<IPropertyHandle> SelectedAudioInputDevice;
+
+	uint32 NumDevices;
+	DeviceInfoArrayHandle->GetNumElements(NumDevices);
+
+	for (uint32 DeviceIndex = 0; DeviceIndex < NumDevices; ++DeviceIndex)
+	{
+		TSharedRef<IPropertyHandle> DeviceInfoHandle = DeviceInfoArrayHandle->GetElement(DeviceIndex);
+		if (DeviceInfoHandle->IsValidHandle())
+		{
+			bool bIsDefaultDevice = GetIsDefaultFromInfoProperty(DeviceInfoHandle);
+			FString DeviceId = GetDeviceIdFromInfoProperty(DeviceInfoHandle);
+
+			if (bIsDefaultDevice && !bHaveCurrentDevice)
+			{
+				CurrentDeviceId = DeviceId;
+				InputDeviceHandle->SetValue(CurrentDeviceId);
+				SelectedAudioInputDevice = DeviceInfoHandle;
+				DeviceChannelCountHandle->SetValue(GetChannelCountFromInfoProperty(DeviceInfoHandle));
+			}
+			else if (bHaveCurrentDevice)
+			{
+				if (CurrentDeviceId == DeviceId)
+				{
+					SelectedAudioInputDevice = DeviceInfoHandle;
+					DeviceChannelCountHandle->SetValue(GetChannelCountFromInfoProperty(DeviceInfoHandle));
+				}
+			}
+
+			AudioInputDevices.Add(DeviceInfoHandle);
+		}
+	}
+
+	// Combo box component:
+	TSharedRef<SComboBox<TSharedPtr<IPropertyHandle>>> NewWidget = SNew(SComboBox<TSharedPtr<IPropertyHandle>>)
+	.OptionsSource(&AudioInputDevices)
+	.OnGenerateWidget_Lambda([](TSharedPtr<IPropertyHandle> InItem)
+	{
+		return SNew(STextBlock)
+			.Text(FText::AsCultureInvariant(GetDeviceNameFromInfoProperty(InItem)))
+			.Font(FAppStyle::Get().GetFontStyle("SmallFont"));
+	})
+	.OnSelectionChanged_Lambda([this](TSharedPtr<IPropertyHandle> InSelection, ESelectInfo::Type InSelectInfo)
+	{
+		if (InSelection.IsValid() && ComboBoxTitleBlock.IsValid())
+		{
+			DeviceChannelCountHandle->SetValue(GetChannelCountFromInfoProperty(InSelection));
+			InputDeviceHandle->SetValue(GetDeviceIdFromInfoProperty(InSelection));
+			ComboBoxTitleBlock->SetText(FText::AsCultureInvariant(GetDeviceNameFromInfoProperty(InSelection)));
+		}
+	})
+	[
+		SAssignNew(ComboBoxTitleBlock, STextBlock)
+		.Text_Lambda([SelectedAudioInputDevice]()
+		{
+			return FText::AsCultureInvariant(GetDeviceNameFromInfoProperty(SelectedAudioInputDevice));
+		})
+		.Font(FAppStyle::Get().GetFontStyle("SmallFont"))
+	];
+
+	ComboBox = NewWidget;
+	UpdateDeviceCombobxEnabled();
+
+	return NewWidget;
+}
+
+void FAudioInputDevicePropertyCustomization::UpdateDeviceCombobxEnabled()
+{
+	// Get current data model value
+	bool bUseSystemDefault;
+	UseSystemDefaultHandle->GetValue(bUseSystemDefault);
+
+	// We want the widget enabled if bUseSystemDefault is false
+	bool bEnableWidget = !bUseSystemDefault;
+
+	// Get the current widget state
+	bool bIsWidgetEnabled = ComboBox->IsEnabled();
+
+	// If the UI doesn't match the data model, refresh the UI
+	if (bEnableWidget != bIsWidgetEnabled)
+	{
+		ComboBox->SetEnabled(bEnableWidget);
+	}
+}
+
+void FAudioInputDevicePropertyCustomization::SynchronizeWidgetStates()
+{
+	UpdateDeviceCombobxEnabled();
+
+	uint32 NumDevices;
+	DeviceInfoArrayHandle->GetNumElements(NumDevices);
+
+	// Get current data model value
+	FString CurrentDeviceId;
+	InputDeviceHandle->GetValue(CurrentDeviceId);
+
+	// Get current menu selection
+	TSharedPtr<IPropertyHandle> CurrentMenuSelection = ComboBox->GetSelectedItem();
+	
+	// If the UI doesn't match the data model, refresh the UI
+	if (CurrentDeviceId != GetDeviceIdFromInfoProperty(CurrentMenuSelection))
+	{
+		for (uint32 DeviceIndex = 0; DeviceIndex < NumDevices; ++DeviceIndex)
+		{
+			TSharedRef<IPropertyHandle> DeviceInfoHandle = DeviceInfoArrayHandle->GetElement(DeviceIndex);
+			if (DeviceInfoHandle->IsValidHandle())
+			{
+				FString DeviceId = GetDeviceIdFromInfoProperty(DeviceInfoHandle);
+
+				if (DeviceId == CurrentDeviceId)
+				{
+					ComboBox->SetSelectedItem(DeviceInfoHandle);
+					break;
+				}
+			}
+		}
+	}
+}
+
+UTakeRecorderAudioInputSettings* TakeRecorderAudioSettingsUtils::GetTakeRecorderAudioInputSettings()
+{
+	const UTakeRecorderProjectSettings* ProjectSettings = GetDefault<UTakeRecorderProjectSettings>();
+	for (TWeakObjectPtr<UObject> WeakObject : ProjectSettings->AdditionalSettings)
+	{
+		UObject* Object = WeakObject.Get();
+		if (Object)
+		{
+			if (Object->IsA<UTakeRecorderAudioInputSettings>())
+			{
+				return Cast<UTakeRecorderAudioInputSettings>(Object);
+			}
+		}
+	}
+	return nullptr;
+}
diff --git a/Engine/Plugins/VirtualProduction/Takes/Source/TakeRecorder/Private/Widgets/TakeRecorderAudioSettingsCustomization.h b/Engine/Plugins/VirtualProduction/Takes/Source/TakeRecorder/Private/Widgets/TakeRecorderAudioSettingsCustomization.h
new file mode 100644
--- /dev/null
+++ b/Engine/Plugins/VirtualProduction/Takes/Source/TakeRecorder/Private/Widgets/TakeRecorderAudioSettingsCustomization.h
@@ -0,0 +1,48 @@
+// Copyright Epic Games, Inc. All Rights Reserved.
+
+#pragma once
+
+#include "DetailWidgetRow.h"
+#include "IPropertyTypeCustomization.h"
+#include "TakeRecorderSourceProperty.h"
+#include "Widgets/Input/SComboBox.h"
+#include "Widgets/Text/STextBlock.h"
+
+class FAudioInputDevicePropertyCustomization : public IPropertyTypeCustomization
+{
+public:
+
+	virtual void CustomizeHeader(TSharedRef<IPropertyHandle> PropertyHandle, FDetailWidgetRow& HeaderRow, IPropertyTypeCustomizationUtils& CustomizationUtils) override {}
+
+	virtual void CustomizeChildren(TSharedRef<IPropertyHandle> PropertyHandle, IDetailChildrenBuilder& ChildBuilder, IPropertyTypeCustomizationUtils& CustomizationUtils) override;
+
+private:
+
+	static FString GetDeviceNameFromInfoProperty(const TSharedPtr<IPropertyHandle>& InDeviceInfoHandle);
+
+	FString GetDeviceIdFromInfoProperty(const TSharedPtr<IPropertyHandle>& InDeviceInfoHandle);
+	int32 GetChannelCountFromInfoProperty(const TSharedPtr<IPropertyHandle>& InDeviceInfoHandle);
+
+	bool GetIsDefaultFromInfoProperty(const TSharedPtr<IPropertyHandle>& InDeviceInfoHandle);
+	TSharedRef<SWidget> MakeAudioInputSelectorWidget();
+
+	void UpdateDeviceCombobxEnabled();
+	void SynchronizeWidgetStates();
+
+	TSharedPtr<IPropertyHandle> UseSystemDefaultHandle;
+	TSharedPtr<IPropertyHandle> InputDeviceHandle;
+	TSharedPtr<IPropertyHandle> DeviceChannelCountHandle;
+	TSharedPtr<IPropertyHandle> BufferSizeHandle;
+	TSharedPtr<IPropertyHandleArray> DeviceInfoArrayHandle;
+
+	TArray<TSharedPtr<IPropertyHandle>> AudioInputDevices;
+	TSharedPtr<SComboBox<TSharedPtr<IPropertyHandle>>> ComboBox;
+	TSharedPtr<STextBlock> ComboBoxTitleBlock;
+};
+
+class TakeRecorderAudioSettingsUtils
+{
+public:
+
+	static UTakeRecorderAudioInputSettings* GetTakeRecorderAudioInputSettings();
+};
\ No newline at end of file
diff --git a/Engine/Plugins/VirtualProduction/Takes/Source/TakeRecorderSources/Private/TakeRecorderMicrophoneAudioManager.cpp b/Engine/Plugins/VirtualProduction/Takes/Source/TakeRecorderSources/Private/TakeRecorderMicrophoneAudioManager.cpp
new file mode 100644
--- /dev/null
+++ b/Engine/Plugins/VirtualProduction/Takes/Source/TakeRecorderSources/Private/TakeRecorderMicrophoneAudioManager.cpp
@@ -0,0 +1,179 @@
+// Copyright Epic Games, Inc. All Rights Reserved.
+
+#include "TakeRecorderMicrophoneAudioManager.h"
+
+#include "TakeRecorderSettings.h"
+#include "TakeRecorderSources.h"
+#include "TakesUtils.h"
+
+#include "IAudioCaptureEditor.h"
+#include "IAudioCaptureEditorModule.h"
+
+#include UE_INLINE_GENERATED_CPP_BY_NAME(TakeRecorderMicrophoneAudioManager)
+
+
+UTakeRecorderMicrophoneAudioManager::UTakeRecorderMicrophoneAudioManager(const FObjectInitializer& ObjInit)
+	: Super(ObjInit)
+{
+}
+
+TUniquePtr<IAudioCaptureEditor> UTakeRecorderMicrophoneAudioManager::CreateAudioRecorderObject()
+{
+	IAudioCaptureEditorModule& Factory = FModuleManager::Get().LoadModuleChecked<IAudioCaptureEditorModule>("AudioCaptureEditor");
+	return Factory.CreateAudioRecorder();
+}
+
+void UTakeRecorderMicrophoneAudioManager::PostEditChangeProperty(FPropertyChangedEvent& PropertyChangedEvent)
+{
+	Super::PostEditChangeProperty(PropertyChangedEvent);
+
+	if (HasAnyFlags(RF_ClassDefaultObject))
+	{
+		SaveConfig();
+	}
+
+	// Notify the Mic sources of the new device channel count
+	GetOnNotifySourcesOfDeviceChange().Broadcast(AudioInputDevice.DeviceInputChannelCount);
+	// Notify the UI of the change so it can update if needed
+	GetOnAudioInputDeviceChanged().Broadcast();
+}
+
+void UTakeRecorderMicrophoneAudioManager::EnumerateAudioDevices(bool InForceRefresh)
+{
+	if (!AudioRecorder.IsValid())
+	{
+		AudioRecorder = CreateAudioRecorderObject();
+	}
+
+	if (AudioInputDevice.DeviceInfoArray.Num() == 0 || InForceRefresh)
+	{
+		BuildDeviceInfoArray();
+	}
+}
+
+void UTakeRecorderMicrophoneAudioManager::StartRecording(int32 InChannelCount)
+{
+	if (InChannelCount <= 0)
+	{
+		UE_LOG(LogTakesCore, Error, TEXT("Microphone Audio Source will not start. No active Mic Sources have been assigned audio device channels."));
+		return;
+	}
+
+	if (AudioRecorder.IsValid())
+	{
+		if (AudioRecorder->IsReadyToRecord())
+		{
+			UE_LOG(LogTakesCore, Verbose, TEXT("Microphone Audio Source AudioRecorder Device: %s"), *(AudioInputDevice.DeviceId));
+
+			FTakeRecorderAudioSettings RecorderSettings;
+			RecorderSettings.AudioCaptureDeviceId = AudioInputDevice.DeviceId;
+			RecorderSettings.NumRecordChannels = InChannelCount;
+			RecorderSettings.AudioInputBufferSize = AudioInputDevice.AudioInputBufferSize;
+
+			AudioRecorder->Start(RecorderSettings);
+		}
+	}
+	else
+	{
+		UE_LOG(LogTakesCore, Error, TEXT("Microphone Audio Source could not start. Please check that the AudioCapture plugin is enabled"));
+	}
+}
+
+void UTakeRecorderMicrophoneAudioManager::StopRecording()
+{
+	if (AudioRecorder.IsValid() && AudioRecorder->IsRecording())
+	{
+		AudioRecorder->Stop();
+	}
+}
+
+TObjectPtr<USoundWave> UTakeRecorderMicrophoneAudioManager::GetRecordedSoundWave(const FTakeRecorderAudioSourceSettings& InSourceSettings)
+{
+	if (AudioRecorder.IsValid() && AudioRecorder->IsStopped())
+	{
+		return AudioRecorder->GetRecordedSoundWave(InSourceSettings);
+	}
+
+	return nullptr;
+}
+
+void UTakeRecorderMicrophoneAudioManager::FinalizeRecording()
+{
+	if (AudioRecorder.IsValid() && !AudioRecorder->IsReadyToRecord())
+	{
+		AudioRecorder.Reset();
+	}
+
+	if (!AudioRecorder.IsValid())
+	{
+		AudioRecorder = CreateAudioRecorderObject();
+		check(AudioRecorder.IsValid());
+	}
+}
+
+void UTakeRecorderMicrophoneAudioManager::BuildDeviceInfoArray()
+{
+	AudioInputDevice.DeviceInfoArray.Empty();
+
+	FTakeRecorderAudioDeviceInfo DefaultDeviceInfo;
+	bool bFoundDefaultDevice = AudioRecorder->GetCaptureDeviceInfo(DefaultDeviceInfo);
+
+	TArray<FTakeRecorderAudioDeviceInfo> CaptureDevicesAvailable;
+	AudioRecorder->GetCaptureDevicesAvailable(CaptureDevicesAvailable);
+	for (const FTakeRecorderAudioDeviceInfo& DeviceInfo : CaptureDevicesAvailable)
+	{
+		bool bIsDefaultDevice = bFoundDefaultDevice && (DeviceInfo.DeviceId == DefaultDeviceInfo.DeviceId);
+		AudioInputDevice.DeviceInfoArray.Add(FAudioInputDeviceInfoProperty(DeviceInfo.DeviceName, DeviceInfo.DeviceId, DeviceInfo.InputChannels, DeviceInfo.PreferredSampleRate, bIsDefaultDevice));
+	}
+
+	// Device names often have numbers in them so perform
+	// an alpha-numeric sort.
+	Algo::Sort(AudioInputDevice.DeviceInfoArray,
+		[](const FAudioInputDeviceInfoProperty& Left, const FAudioInputDeviceInfoProperty& Right) -> bool
+		{
+			auto StrToNum = [](const FString& InString, int32& InIndex) -> int32
+			{
+				const int32 StrLen = InString.Len();
+				int32 Value = InString[InIndex] - '0';
+
+				while ((InIndex + 1) < StrLen && FChar::IsDigit(InString[InIndex + 1]))
+				{
+					Value = (Value * 10) + (InString[++InIndex] - '0');
+				}
+
+				return Value;
+			};
+
+			const FString& LeftString = Left.DeviceName;
+			const FString& RightString = Right.DeviceName;
+			const int32 LeftLen = LeftString.Len();
+			const int32 RightLen = RightString.Len();
+			int32 LeftIndex = 0;
+			int32 RightIndex = 0;
+
+			while (LeftIndex < LeftLen && RightIndex < RightLen)
+			{
+				FString::ElementType LeftChar = LeftString[LeftIndex];
+				FString::ElementType RightChar = RightString[RightIndex];
+				if (FChar::IsDigit(LeftChar) && FChar::IsDigit(RightChar))
+				{
+					int32 LeftValue = StrToNum(LeftString, LeftIndex);
+					int32 RightValue = StrToNum(RightString, RightIndex);
+
+					if (LeftValue != RightValue)
+					{
+						return LeftValue < RightValue;
+					}
+				}
+				else if (LeftChar != RightChar)
+				{
+					return LeftChar < RightChar;
+				}
+
+				++LeftIndex;
+				++RightIndex;
+			}
+
+			return LeftLen < RightLen;
+		});
+}
diff --git a/Engine/Plugins/VirtualProduction/Takes/Source/TakeRecorderSources/Private/TakeRecorderMicrophoneAudioManager.h b/Engine/Plugins/VirtualProduction/Takes/Source/TakeRecorderSources/Private/TakeRecorderMicrophoneAudioManager.h
new file mode 100644
--- /dev/null
+++ b/Engine/Plugins/VirtualProduction/Takes/Source/TakeRecorderSources/Private/TakeRecorderMicrophoneAudioManager.h
@@ -0,0 +1,77 @@
+// Copyright Epic Games, Inc. All Rights Reserved.
+
+#pragma once
+
+#include "TakeRecorderSourceProperty.h"
+
+#include "TakeRecorderMicrophoneAudioManager.generated.h"
+
+class IAudioCaptureEditor;
+class USoundWave;
+struct FTakeRecorderAudioSourceSettings;
+
+
+DECLARE_MULTICAST_DELEGATE_OneParam(FOnNotifySourcesOfDeviceChange, int);
+
+/** This class exposes the audio input device list via the project settings details. It does this in 
+*   conjunction with FAudioInputDevicePropertyCustomization. It also manages the IAudioCaptureEditor 
+*   object which handles the low level audio device recording.
+*/
+UCLASS(config=EditorPerProjectUserSettings, PerObjectConfig, DisplayName = "Audio Input Device")
+class TAKERECORDERSOURCES_API UTakeRecorderMicrophoneAudioManager : public UTakeRecorderAudioInputSettings
+{
+public:
+	GENERATED_BODY()
+
+	UTakeRecorderMicrophoneAudioManager(const FObjectInitializer& ObjInit);
+
+	// Begin UObject Interface
+	virtual void PostEditChangeProperty(FPropertyChangedEvent& PropertyChangedEvent) override;
+	// End UObject Interface
+
+	/** Enumerates the audio devices present on the current machine */
+	virtual void EnumerateAudioDevices(bool InForceRefresh = false) override;
+
+	/** Returns input channel count for currently selected audio device */
+	virtual int32 GetDeviceChannelCount() override { return AudioInputDevice.DeviceInputChannelCount;  }
+
+	/** The audio device to use for this microphone source */
+	UPROPERTY(config, EditAnywhere, BlueprintReadWrite, Category = "Source", meta = (ShowOnlyInnerProperties))
+	FAudioInputDeviceProperty AudioInputDevice;
+
+	/** 
+	*  Calls StartRecording on the AudioRecorder object. This is called multiple times (once for each
+	*  microphone source, however, only the first call triggers the call to AudioRecorder.
+	*/
+	void StartRecording(int32 InChannelCount);
+	/**
+	 * Calls StopRecording on the AudioRecorder object. This is called multiple times (once for each
+	 * microphone source, however, only the first call triggers the call to AudioRecorder.
+	 */
+	void StopRecording();
+	/**
+	 *  Calls FinalizeRecording on the AudioRecorder object. This is called multiple times (once for each
+	 *  microphone source, however, only the first call triggers the call to AudioRecorder.
+	 */
+	void FinalizeRecording();
+
+
+	/** Fetches the USoundWave for this source after a Take has been recorded */
+	TObjectPtr<class USoundWave> GetRecordedSoundWave(const FTakeRecorderAudioSourceSettings& InSourceSettings);
+	/** Gets the cached value for the audio input device channel count */
+	int32 GetDeviceChannelCount() const { return AudioInputDevice.DeviceInputChannelCount; }
+	/** Accessor for the OnNotifySourcesOfDeviceChange delegate list */
+	FOnNotifySourcesOfDeviceChange& GetOnNotifySourcesOfDeviceChange() { return OnNotifySourcesOfDeviceChange; }
+
+private:
+
+	/** Multicast delegate which notifies clients when the currently selected audio device changes. */
+	FOnNotifySourcesOfDeviceChange OnNotifySourcesOfDeviceChange;
+	/** Calls factory to create the AudioRecorder object */
+	TUniquePtr<IAudioCaptureEditor> CreateAudioRecorderObject();
+	/** Builds the list of audio devices which will be used in the device menu */
+	void BuildDeviceInfoArray();
+
+	/** The audio recorder object which manages low level recording of audio data */
+	TUniquePtr<IAudioCaptureEditor> AudioRecorder;
+};
diff --git a/Engine/Plugins/VirtualProduction/Takes/Source/TakeRecorderSources/Private/TakeRecorderMicrophoneAudioSource.cpp b/Engine/Plugins/VirtualProduction/Takes/Source/TakeRecorderSources/Private/TakeRecorderMicrophoneAudioSource.cpp
--- a/Engine/Plugins/VirtualProduction/Takes/Source/TakeRecorderSources/Private/TakeRecorderMicrophoneAudioSource.cpp
+++ b/Engine/Plugins/VirtualProduction/Takes/Source/TakeRecorderSources/Private/TakeRecorderMicrophoneAudioSource.cpp
@@ -1,28 +1,29 @@
 // Copyright Epic Games, Inc. All Rights Reserved.
 
 #include "TakeRecorderMicrophoneAudioSource.h"
-#include "AssetRegistry/AssetData.h"
+#include "TakeRecorderMicrophoneAudioManager.h"
 #include "TakeRecorderSources.h"
 #include "TakeRecorderSettings.h"
 #include "TakesUtils.h"
 #include "TakeMetaData.h"
-#include "Recorder/TakeRecorderParameters.h"
-#include "Styling/SlateIconFinder.h"
-#include "LevelSequence.h"
+
+#include "AudioCaptureEditorTypes.h"
 #include "Editor.h"
+#include "LevelSequence.h"
 #include "Modules/ModuleManager.h"
+#include "MovieSceneFolder.h"
+#include "Recorder/TakeRecorderParameters.h"
 #include "Sections/MovieSceneAudioSection.h"
-#include "Tracks/MovieSceneAudioTrack.h"
 #include "Sound/SoundWave.h"
-#include "ISequenceAudioRecorder.h"
-#include "ISequenceRecorder.h"
-#include "MovieSceneFolder.h"
+#include "Styling/SlateIconFinder.h"
+#include "Tracks/MovieSceneAudioTrack.h"
 
-#include "Misc/PackageName.h"
 #include "AssetRegistry/AssetData.h"
 #include "AssetRegistry/AssetRegistryModule.h"
-#include "UObject/UObjectBaseUtility.h"
+#include "Misc/PackageName.h"
+#include "ObjectEditorUtils.h"
 #include "ObjectTools.h"
+#include "UObject/UObjectBaseUtility.h"
 
 #include UE_INLINE_GENERATED_CPP_BY_NAME(TakeRecorderMicrophoneAudioSource)
 
@@ -66,11 +67,81 @@ FString UTakeRecorderMicrophoneAudioSourceSettings::GetSubsceneAssetName(ULevelS
 UTakeRecorderMicrophoneAudioSource::UTakeRecorderMicrophoneAudioSource(const FObjectInitializer& ObjInit)
 	: Super(ObjInit)
 	, AudioGain(0.0f)
-	, bSplitAudioChannelsIntoSeparateTracks(false)
 	, bReplaceRecordedAudio(true)
 {
 }
 
+void UTakeRecorderMicrophoneAudioSource::Initialize()
+{
+	UTakeRecorderMicrophoneAudioManager* AudioInputManager = GetAudioInputManager();
+	if (AudioInputManager != nullptr)
+	{
+		TUniquePtr<TArray<bool>> ChannelsInUse = GetChannelsInUse(AudioInputManager->GetDeviceChannelCount());
+		int32 AvailableChannelNumber = INDEX_NONE;
+
+		for (int32 ChannelIndex = 0; ChannelIndex < ChannelsInUse->Num(); ++ChannelIndex)
+		{
+			if (!(*ChannelsInUse)[ChannelIndex])
+			{
+				AvailableChannelNumber = ChannelIndex + 1;
+				break;
+			}
+		}
+
+		if (AvailableChannelNumber != INDEX_NONE)
+		{
+			SetCurrentInputChannel(AvailableChannelNumber);
+		}
+	}
+}
+
+void UTakeRecorderMicrophoneAudioSource::PostEditChangeProperty(FPropertyChangedEvent& PropertyChangedEvent)
+{
+	Super::PostEditChangeProperty(PropertyChangedEvent);
+}
+
+UTakeRecorderMicrophoneAudioManager* UTakeRecorderMicrophoneAudioSource::GetAudioInputManager()
+{
+	return GetMutableDefault<UTakeRecorderMicrophoneAudioManager>();
+}
+
+void UTakeRecorderMicrophoneAudioSource::OnNotifySourcesOfDeviceChange(int32 InChannelCount)
+{
+	SetAudioDeviceChannelCount(InChannelCount);
+}
+
+TUniquePtr<TArray<bool>> UTakeRecorderMicrophoneAudioSource::GetChannelsInUse(const int32 InDeviceChannelCount)
+{
+	TUniquePtr<TArray<bool>> ChannelsInUse = MakeUnique<TArray<bool>>();
+	ChannelsInUse->Init(false, InDeviceChannelCount);
+
+	UTakeRecorderSources* SourcesList = GetTypedOuter<UTakeRecorderSources>();
+	for (UTakeRecorderSource* Source : SourcesList->GetSources())
+	{
+		if (UTakeRecorderMicrophoneAudioSource* MicSource = Cast<UTakeRecorderMicrophoneAudioSource>(Source))
+		{
+			int32 MicSourceChannelNumber = MicSource->AudioChannel.AudioInputDeviceChannel;
+			int32 ChannelIndex = MicSourceChannelNumber - 1;
+			if (ChannelIndex >= 0)
+			{
+				(*ChannelsInUse)[ChannelIndex] = true;
+			}
+		}
+	}
+
+	return ChannelsInUse;
+}
+
+void UTakeRecorderMicrophoneAudioSource::SetAudioDeviceChannelCount(int32 InChannelCount)
+{
+	if (ensure(InChannelCount >= 0))
+	{
+		if (AudioChannel.AudioInputDeviceChannel > InChannelCount)
+		{
+			AudioChannel.AudioInputDeviceChannel = 0;
+		}
+	}
+}
 
 static FString MakeNewAssetName(const FString& BaseAssetPath, const FString& BaseAssetName)
 {
@@ -100,31 +171,36 @@ static FString MakeNewAssetName(const FString& BaseAssetPath, const FString& Bas
 
 TArray<UTakeRecorderSource*> UTakeRecorderMicrophoneAudioSource::PreRecording(ULevelSequence* InSequence, FMovieSceneSequenceID InSequenceID, ULevelSequence* InRootSequence, FManifestSerializer* InManifestSerializer)
 {
-	UMovieScene* MovieScene = InSequence->GetMovieScene();
-	for (auto Track : MovieScene->GetTracks())
+	if (AudioChannel.AudioInputDeviceChannel > 0)
 	{
-		if (Track->IsA(UMovieSceneAudioTrack::StaticClass()) && Track->GetDisplayName().EqualTo(AudioTrackName))
+		UMovieScene* MovieScene = InSequence->GetMovieScene();
+		for (auto Track : MovieScene->GetTracks())
 		{
-			CachedAudioTrack = Cast<UMovieSceneAudioTrack>(Track);
+			if (Track->IsA(UMovieSceneAudioTrack::StaticClass()) && Track->GetDisplayName().EqualTo(AudioTrackName))
+			{
+				CachedAudioTrack = Cast<UMovieSceneAudioTrack>(Track);
+			}
 		}
-	}
 
-	if (!CachedAudioTrack.IsValid())
-	{
-		CachedAudioTrack = MovieScene->AddTrack<UMovieSceneAudioTrack>();
-		CachedAudioTrack->SetDisplayName(AudioTrackName);
-	}
+		if (!CachedAudioTrack.IsValid())
+		{
+			CachedAudioTrack = MovieScene->AddTrack<UMovieSceneAudioTrack>();
+			CachedAudioTrack->SetDisplayName(AudioTrackName);
+		}
 
-	FString PathToRecordTo = FPackageName::GetLongPackagePath(InSequence->GetOutermost()->GetPathName());
-	FString BaseName = InSequence->GetName();
+		FString PathToRecordTo = FPackageName::GetLongPackagePath(InSequence->GetOutermost()->GetPathName());
+		FString BaseName = InSequence->GetName();
 
-	AudioDirectory.Path = PathToRecordTo;
-	if (AudioSubDirectory.Len())
-	{
-		AudioDirectory.Path /= AudioSubDirectory;
+		AudioDirectory.Path = PathToRecordTo;
+		if (AudioSubDirectory.Len())
+		{
+			AudioDirectory.Path /= AudioSubDirectory;
+		}
+
+		AssetName = MakeNewAssetName(AudioDirectory.Path, BaseName);
 	}
 
-	AssetName = MakeNewAssetName(AudioDirectory.Path, BaseName);
+	RecordedSoundWaves.Empty();
 
 	return TArray<UTakeRecorderSource*>();
 }
@@ -142,25 +218,24 @@ void UTakeRecorderMicrophoneAudioSource::StartRecording(const FTimecode& InSecti
 {
 	Super::StartRecording(InSectionStartTimecode, InSectionFirstFrame, InSequence);
 
-	ISequenceRecorder& Recorder = FModuleManager::Get().LoadModuleChecked<ISequenceRecorder>("SequenceRecorder");
-
-	FSequenceAudioRecorderSettings AudioSettings;
-	AudioSettings.Directory = AudioDirectory;
-	AudioSettings.AssetName = AssetName;
-	AudioSettings.GainDb = AudioGain;
-	AudioSettings.bSplitChannels = bSplitAudioChannelsIntoSeparateTracks;
+	StartTimecode = InSectionStartTimecode;
 
-	RecordedSoundWaves.Empty();
-	
-	AudioRecorder = Recorder.CreateAudioRecorder();
-	if (AudioRecorder)
-	{
-		UE_LOG(LogTakesCore, Verbose, TEXT("Microphone Audio Source AudioRecorder Start: %s"), *AssetName);
-		AudioRecorder->Start(AudioSettings);
-	}
-	else
+	UTakeRecorderMicrophoneAudioManager* AudioInputManager = GetAudioInputManager();
+	if (AudioInputManager != nullptr)
 	{
-		UE_LOG(LogTakesCore, Error, TEXT("Microphone Audio Source could not start. Please check that the AudioCapture plugin is enabled"));
+		int32 LastChannelInUse = INDEX_NONE;
+		TUniquePtr<TArray<bool>> ChannelsInUse = GetChannelsInUse(AudioInputManager->GetDeviceChannelCount());
+
+		for (int32 ChannelIndex = 0; ChannelIndex < ChannelsInUse->Num(); ++ChannelIndex)
+		{
+			if ((*ChannelsInUse)[ChannelIndex])
+			{
+				LastChannelInUse = FMath::Max(LastChannelInUse, ChannelIndex + 1);
+			}
+		}
+
+
+		AudioInputManager->StartRecording(LastChannelInUse);
 	}
 }
 
@@ -168,21 +243,17 @@ void UTakeRecorderMicrophoneAudioSource::StopRecording(class ULevelSequence* InS
 {
 	Super::StopRecording(InSequence);
 
-	if (AudioRecorder)
+	UTakeRecorderMicrophoneAudioManager* AudioInputManager = GetAudioInputManager();
+	if (AudioInputManager != nullptr)
 	{
-		TArray<USoundWave*> SoundWaves;
-		AudioRecorder->Stop(SoundWaves);
-
-		for (USoundWave* RecordedSoundWave : SoundWaves)
-		{
-			RecordedSoundWaves.Add(RecordedSoundWave);
-		}
+		AudioInputManager->StopRecording();
 	}
-	AudioRecorder.Reset();
 }
 
 TArray<UTakeRecorderSource*> UTakeRecorderMicrophoneAudioSource::PostRecording(ULevelSequence* InSequence, class ULevelSequence* InRootSequence, const bool bCancelled)
 {
+	GetRecordedSoundWave(InSequence);
+
 	if (!RecordedSoundWaves.Num())
 	{
 		return TArray<UTakeRecorderSource*>();
@@ -265,6 +336,45 @@ TArray<UTakeRecorderSource*> UTakeRecorderMicrophoneAudioSource::PostRecording(U
 	return TArray<UTakeRecorderSource*>();
 }
 
+void UTakeRecorderMicrophoneAudioSource::FinalizeRecording()
+{
+	UTakeRecorderMicrophoneAudioManager* AudioInputManager = GetAudioInputManager();
+	if (AudioInputManager != nullptr)
+	{
+		AudioInputManager->FinalizeRecording();
+	}
+}
+
+void UTakeRecorderMicrophoneAudioSource::GetRecordedSoundWave(ULevelSequence* InSequence)
+{
+	UTakeRecorderMicrophoneAudioManager* AudioInputManager = GetAudioInputManager();
+	if (AudioChannel.AudioInputDeviceChannel > 0 && AudioInputManager != nullptr)
+	{
+		UMovieScene* MovieScene = InSequence->GetMovieScene();
+		FFrameRate DisplayRate = MovieScene->GetDisplayRate();
+
+		FTakeRecorderAudioSourceSettings AudioSettings;
+		AudioSettings.Directory = AudioDirectory;
+		AudioSettings.AssetName = AssetName;
+		AudioSettings.GainDb = AudioGain;
+		AudioSettings.InputChannelNumber = AudioChannel.AudioInputDeviceChannel;
+		AudioSettings.StartTimecode = StartTimecode;
+		AudioSettings.VideoFrameRate = DisplayRate;
+
+		TObjectPtr<USoundWave> SoundWave = AudioInputManager->GetRecordedSoundWave(AudioSettings);
+		if (SoundWave != nullptr)
+		{
+			FSoundWaveTimecodeInfo TimecodeInfo;
+			FSoundTimecodeOffset TimecodeOffset;
+
+			TimecodeOffset.NumOfSecondsSinceMidnight = StartTimecode.ToTimespan(DisplayRate).GetTotalSeconds();
+			SoundWave->SetTimecodeOffset(TimecodeOffset);
+			
+			RecordedSoundWaves.Add(SoundWave);
+		}
+	}
+}
+
 FText UTakeRecorderMicrophoneAudioSource::GetDisplayTextImpl() const
 {
 	return NSLOCTEXT("UTakeRecorderMicrophoneAudioSource", "Label", "Microphone Audio");
@@ -272,13 +382,21 @@ FText UTakeRecorderMicrophoneAudioSource::GetDisplayTextImpl() const
 
 bool UTakeRecorderMicrophoneAudioSource::CanAddSource(UTakeRecorderSources* InSources) const
 {
+	int32 MicrophoneSourceCount = 0;
 	for (UTakeRecorderSource* Source : InSources->GetSources())
 	{
 		if (Source->IsA<UTakeRecorderMicrophoneAudioSource>())
 		{
-			return false;
+			++MicrophoneSourceCount;
 		}
 	}
-	return true;
+	return MicrophoneSourceCount < FAudioInputDeviceProperty::MaxInputChannelCount;
 }
 
+void UTakeRecorderMicrophoneAudioSource::SetCurrentInputChannel(int32 InChannelNumber)
+{
+	if (ensure(InChannelNumber >= 0 && InChannelNumber <= FAudioInputDeviceProperty::MaxInputChannelCount))
+	{
+		AudioChannel.AudioInputDeviceChannel = InChannelNumber;
+	}
+}
diff --git a/Engine/Plugins/VirtualProduction/Takes/Source/TakeRecorderSources/Private/TakeRecorderMicrophoneAudioSource.h b/Engine/Plugins/VirtualProduction/Takes/Source/TakeRecorderSources/Private/TakeRecorderMicrophoneAudioSource.h
--- a/Engine/Plugins/VirtualProduction/Takes/Source/TakeRecorderSources/Private/TakeRecorderMicrophoneAudioSource.h
+++ b/Engine/Plugins/VirtualProduction/Takes/Source/TakeRecorderSources/Private/TakeRecorderMicrophoneAudioSource.h
@@ -4,15 +4,21 @@
 
 #include "TakeRecorderSource.h"
 #include "TakeRecorderSources.h"
-#include "UObject/SoftObjectPtr.h"
-#include "Templates/SubclassOf.h"
+#include "TakeRecorderSourceProperty.h"
+
 #include "GameFramework/Actor.h"
-#include "ISequenceAudioRecorder.h"
+#include "Templates/SubclassOf.h"
+#include "UObject/SoftObjectPtr.h"
+
 #include "TakeRecorderMicrophoneAudioSource.generated.h"
 
+class USoundWave;
+class UTakeRecorderMicrophoneAudioManager;
+struct FTakeRecorderAudioSourceSettings;
+
 /** A recording source that records microphone audio */
 UCLASS(Abstract, config=EditorSettings, DisplayName="Microphone Audio Recorder")
-class UTakeRecorderMicrophoneAudioSourceSettings : public UTakeRecorderSource
+class TAKERECORDERSOURCES_API UTakeRecorderMicrophoneAudioSourceSettings : public UTakeRecorderSource
 {
 public:
 	GENERATED_BODY()
@@ -41,45 +47,75 @@ public:
 
 /** A recording source that records microphone audio */
 UCLASS(Category="Audio", config=EditorSettings, meta = (TakeRecorderDisplayName = "Microphone Audio"))
-class UTakeRecorderMicrophoneAudioSource : public UTakeRecorderMicrophoneAudioSourceSettings
+class TAKERECORDERSOURCES_API UTakeRecorderMicrophoneAudioSource : public UTakeRecorderMicrophoneAudioSourceSettings
 {
 public:
 	GENERATED_BODY()
 
 	UTakeRecorderMicrophoneAudioSource(const FObjectInitializer& ObjInit);
 
+	// Begin UObject Interface
+	virtual void PostEditChangeProperty(FPropertyChangedEvent& PropertyChangedEvent) override;
+	// End UObject Interface
+
 	/** Gain in decibels to apply to recorded audio */
 	UPROPERTY(EditAnywhere, BlueprintReadWrite, Category = "Source", meta = (ClampMin = "0.0", UIMin = "0.0"))
 	float AudioGain;
 
 	/** Whether or not to split mic channels into separate audio tracks. If not true, a max of 2 input channels is supported. */
-	UPROPERTY(EditAnywhere, BlueprintReadWrite, Category = "Source")
-	bool bSplitAudioChannelsIntoSeparateTracks;
+#if WITH_EDITORONLY_DATA
+	UPROPERTY(meta = (DeprecatedProperty, DeprecationMessage = "SplitAudioChannelsIntoSeparateTracks is deprecated."))
+	bool bSplitAudioChannelsIntoSeparateTracks_DEPRECATED;
+#endif
 
 	/** Replace existing recorded audio with any newly recorded audio */
 	UPROPERTY(EditAnywhere, BlueprintReadWrite, Category = "Source")
 	bool bReplaceRecordedAudio;
 
-private:
+	/** The audio device to use for this microphone source */
+	UPROPERTY(EditAnywhere, BlueprintReadWrite, Category = "Source", meta = (ShowOnlyInnerProperties))
+	FAudioInputDeviceChannelProperty AudioChannel;
 
 	// UTakeRecorderSource
+	virtual void Initialize() override;
 	virtual TArray<UTakeRecorderSource*> PreRecording(ULevelSequence* InSequence, FMovieSceneSequenceID InSequenceID, ULevelSequence* InRootSequence, FManifestSerializer* InManifestSerializer) override;
 	virtual void StartRecording(const FTimecode& InSectionStartTimecode, const FFrameNumber& InSectionFirstFrame, class ULevelSequence* InSequence) override;
 	virtual void StopRecording(class ULevelSequence* InSequence) override;
 	virtual TArray<UTakeRecorderSource*> PostRecording(ULevelSequence* InSequence, class ULevelSequence* InRootSequence, const bool bCancelled) override;
+	virtual void FinalizeRecording() override;
 	virtual FText GetDisplayTextImpl() const override;
 	virtual void AddContentsToFolder(class UMovieSceneFolder* InFolder) override;
 	virtual bool CanAddSource(UTakeRecorderSources* InSources) const override;
 	// ~UTakeRecorderSource
 
+	/** Sets the channel count supported by the currently selected audio device. */
+	void SetAudioDeviceChannelCount(int32 InChannelCount);
+	/** Delegate which receives notifications when the audio input device changes. */
+	void OnNotifySourcesOfDeviceChange(int32 InChannelCount);
+
 private:
-	TWeakObjectPtr<class UMovieSceneAudioTrack> CachedAudioTrack;
 
+	/** Helper function for getting a pointer to the AudioInputManger object. */
+	static UTakeRecorderMicrophoneAudioManager* GetAudioInputManager();
+	
+	/** Returns an array of booleans indicating which channel indexes are currently in use */
+	TUniquePtr<TArray<bool>> GetChannelsInUse(const int32 InDeviceChannelCount);
+	/** Fetches the USoundWave for this source after a Take has been recorded */
+	void GetRecordedSoundWave(ULevelSequence* InSequence);
+	/** Called when the user selects a new channel in the combobox for this source */
+	void SetCurrentInputChannel(int32 InChannelNumber);
+
+private:
+
+	// Holds the Sequencer audio track for this source
+	TWeakObjectPtr<class UMovieSceneAudioTrack> CachedAudioTrack;
+	// Holds the USoundWave assets for a given take
 	TArray<TWeakObjectPtr<USoundWave>> RecordedSoundWaves;
+	// Caches the starting timecode for this take so it can be referenced when creating USoundWave assets
+	FTimecode StartTimecode;
 
-	TUniquePtr<ISequenceAudioRecorder> AudioRecorder;
-	
-	//Created in PreRecord but used in StartRecording.
+	// The user specified directory to store recorded audio assets in
 	FDirectoryPath AudioDirectory;
+	// The name of the UsoundWave asset 
 	FString AssetName;
 };
diff --git a/Engine/Plugins/VirtualProduction/Takes/Source/TakeRecorderSources/Private/TakeRecorderSourcesModule.cpp b/Engine/Plugins/VirtualProduction/Takes/Source/TakeRecorderSources/Private/TakeRecorderSourcesModule.cpp
--- a/Engine/Plugins/VirtualProduction/Takes/Source/TakeRecorderSources/Private/TakeRecorderSourcesModule.cpp
+++ b/Engine/Plugins/VirtualProduction/Takes/Source/TakeRecorderSources/Private/TakeRecorderSourcesModule.cpp
@@ -14,6 +14,7 @@
 #include "TakeMetaData.h"
 #include "TakeRecorderSourceHelpers.h"
 #include "TakeRecorderActorSource.h"
+#include "TakeRecorderMicrophoneAudioManager.h"
 #include "TakeRecorderSources.h"
 #include "TakeRecorderSourcesCommands.h"
 #include "TakeRecorderSettings.h"
@@ -320,6 +321,7 @@ public:
 		SourcesMenuExtension = TakeRecorderModule.RegisterSourcesMenuExtension(FOnExtendSourcesMenu::CreateStatic(ExtendSourcesMenu));
 
 		TakeRecorderModule.RegisterSettingsObject(GetMutableDefault<UTakeRecorderMicrophoneAudioSourceSettings>());
+		TakeRecorderModule.RegisterSettingsObject(GetMutableDefault<UTakeRecorderMicrophoneAudioManager>());
 		TakeRecorderModule.RegisterSettingsObject(GetMutableDefault<UMovieSceneAnimationTrackRecorderEditorSettings>());
 		TakeRecorderModule.RegisterSettingsObject(GetMutableDefault<UTakeRecorderWorldSourceSettings>());
 
diff --git a/Engine/Plugins/VirtualProduction/Takes/Source/TakeRecorderSources/TakeRecorderSources.Build.cs b/Engine/Plugins/VirtualProduction/Takes/Source/TakeRecorderSources/TakeRecorderSources.Build.cs
--- a/Engine/Plugins/VirtualProduction/Takes/Source/TakeRecorderSources/TakeRecorderSources.Build.cs
+++ b/Engine/Plugins/VirtualProduction/Takes/Source/TakeRecorderSources/TakeRecorderSources.Build.cs
@@ -13,7 +13,8 @@ public class TakeRecorderSources : ModuleRules
 				"CoreUObject",
 				"EditorFramework",
 				"Engine",
-				
+
+				"AudioCaptureEditor",
 				"InputCore",
 				"LevelEditor",
 				"LevelSequence",
@@ -21,7 +22,7 @@ public class TakeRecorderSources : ModuleRules
                 "MovieScene",
 				"MovieSceneTracks",
                 "SceneOutliner",
-				"SequenceRecorder", // For ISequenceAudioRecorder
+				"SequenceRecorder", // For FTimecodeBoneMethod
 				"SerializedRecorderInterface",
                 "Slate",
 				"SlateCore",
diff --git a/Engine/Plugins/VirtualProduction/Takes/Source/TakesCore/Private/TakeRecorderSources.cpp b/Engine/Plugins/VirtualProduction/Takes/Source/TakesCore/Private/TakeRecorderSources.cpp
--- a/Engine/Plugins/VirtualProduction/Takes/Source/TakesCore/Private/TakeRecorderSources.cpp
+++ b/Engine/Plugins/VirtualProduction/Takes/Source/TakesCore/Private/TakeRecorderSources.cpp
@@ -48,6 +48,7 @@ UTakeRecorderSource* UTakeRecorderSources::AddSource(TSubclassOf<UTakeRecorderSo
 		NewSource = NewObject<UTakeRecorderSource>(this, Class, NAME_None, RF_Transactional);
 		if (ensure(NewSource))
 		{
+			NewSource->Initialize();
 			Sources.Add(NewSource);
 			++SourcesSerialNumber;
 		}
diff --git a/Engine/Plugins/VirtualProduction/Takes/Source/TakesCore/Public/TakeRecorderSource.h b/Engine/Plugins/VirtualProduction/Takes/Source/TakesCore/Public/TakeRecorderSource.h
--- a/Engine/Plugins/VirtualProduction/Takes/Source/TakesCore/Public/TakeRecorderSource.h
+++ b/Engine/Plugins/VirtualProduction/Takes/Source/TakesCore/Public/TakeRecorderSource.h
@@ -44,6 +44,13 @@ public:
 	FColor TrackTint;
 
 public:
+	/**
+	 * Gives the Source object an opportunity to perform additional initialization outside of the constructor. This can be necessary
+	 * when needing to perform additional initialization that should not occur if/when the object is duplicated (as is the case when 
+	 * all sources are duplicated prior to each record pass).
+	 */
+	virtual void Initialize() {}
+
 	/** 
 	* If the source is Valid and ready to be recorded. Tells the Take Recorder that it's really going to record. Note that PreRecord/StartRecord
 	* may still get called on the source even if not Valid. This is used mainly so we only Start Recording if we have more than one Valid Source.
diff --git a/Engine/Plugins/VirtualProduction/Takes/Source/TakesCore/Public/TakeRecorderSourceProperty.h b/Engine/Plugins/VirtualProduction/Takes/Source/TakesCore/Public/TakeRecorderSourceProperty.h
--- a/Engine/Plugins/VirtualProduction/Takes/Source/TakesCore/Public/TakeRecorderSourceProperty.h
+++ b/Engine/Plugins/VirtualProduction/Takes/Source/TakesCore/Public/TakeRecorderSourceProperty.h
@@ -94,3 +94,120 @@ private:
 
 	Cache RecordingInfo;
 };
+
+/**
+ * Encapsulates audio device properties which are utilized by UI facing classes such as FAudioInputDeviceProperty.
+ */
+USTRUCT(BlueprintType)
+struct TAKESCORE_API FAudioInputDeviceInfoProperty
+{
+	GENERATED_BODY()
+
+	FAudioInputDeviceInfoProperty() = default;
+	FAudioInputDeviceInfoProperty(const FAudioInputDeviceInfoProperty& InObject) = default;
+
+	FAudioInputDeviceInfoProperty(const FString& InDeviceName, const FString& InDeviceId, int32 InChannelCount, int32 InSampleRate, bool bInIsDefaultDevice)
+		: DeviceName(InDeviceName)
+		, DeviceId(InDeviceId)
+		, InputChannels(InChannelCount)
+		, PreferredSampleRate(InSampleRate)
+		, bIsDefaultDevice(bInIsDefaultDevice)
+	{
+	}
+
+	/** User friendly name of the audio device */
+	UPROPERTY(Transient, VisibleAnywhere, BlueprintReadOnly, Category="Property")
+	FString DeviceName;
+
+	/** The unique id used to identify the device */
+	UPROPERTY(VisibleAnywhere, BlueprintReadOnly, Category="Property")
+	FString DeviceId;
+
+	/** The number input channels this device supports */
+	UPROPERTY(Transient, VisibleAnywhere, BlueprintReadOnly, Category = "Property")
+	int32 InputChannels = 0;
+
+	/** The preferred sample rate for this audio device */
+	UPROPERTY(Transient, VisibleAnywhere, BlueprintReadOnly, Category = "Property")
+	int32 PreferredSampleRate = 0;
+
+	/** Boolean indicating if this device is the currently the system selected input device */
+	UPROPERTY(Transient, VisibleAnywhere, BlueprintReadOnly, Category = "Property")
+	bool bIsDefaultDevice = false;
+};
+
+/**
+ * Encapsulates the array of audio input devices which is populated by UTakeRecorderMicrophoneAudioManager and
+ * utilized by the audio input device list in FAudioInputDevicePropertyCustomization.
+ */
+USTRUCT(BlueprintType)
+struct TAKESCORE_API FAudioInputDeviceProperty
+{
+	GENERATED_BODY()
+
+	/** The maximum number input channels currently supported */
+	constexpr static int32 MaxInputChannelCount = 8;
+
+	/** Boolean indicating if the system selects audio device should be used or to use the selected device from the details panel */
+	UPROPERTY(config, EditAnywhere, BlueprintReadWrite, Category = "Property")
+	bool bUseSystemDefaultAudioDevice = true;
+
+	/* Holds device information for each of the audio devices available on this system. */
+	UPROPERTY(Transient, VisibleAnywhere, BlueprintReadOnly, Category = "Property")
+	TArray<FAudioInputDeviceInfoProperty> DeviceInfoArray;
+
+	/** The unique id of the currently selected audio device */
+	UPROPERTY(config, EditAnywhere, BlueprintReadWrite, Category = "Property", meta = (DisplayName = "Audio Input Device"))
+	FString DeviceId;
+
+	/** The number of input channels supported by this device */
+	UPROPERTY(Transient, EditAnywhere, BlueprintReadWrite, Category = "Property")
+	int32 DeviceInputChannelCount;
+	
+	/** The desired buffer size used for audio callbacks during record */
+	UPROPERTY(config, EditAnywhere, BlueprintReadWrite, Category = "Property", meta = (ClampMin = "256", UIMin = "256", ClampMax = "8192", UIMax = "8192"))
+	int32 AudioInputBufferSize = 1024;	
+};
+
+/**
+ * This class is used by Microphone sources to track the currently selected channel. It aslo
+ * contains a local cache of the max channel count for the currently selected audio device.
+ */
+USTRUCT(BlueprintType)
+struct TAKESCORE_API FAudioInputDeviceChannelProperty
+{
+	GENERATED_BODY()
+
+public:
+
+	/** The currently selected channel from the details panel for this source */
+	UPROPERTY(EditAnywhere, BlueprintReadWrite, Category = "Property")
+	int32 AudioInputDeviceChannel;
+};
+
+DECLARE_MULTICAST_DELEGATE(FOnAudioInputDeviceChanged);
+
+/**
+ * Provides access to the UI code for registering OnAudioInputDeviceChanged delegates.
+ */
+UCLASS(Abstract, Blueprintable, BlueprintType)
+class TAKESCORE_API UTakeRecorderAudioInputSettings : public UObject
+{
+	GENERATED_BODY()
+	
+public:
+
+	/** Subclasses can override this to enumerate audio devices */
+	virtual void EnumerateAudioDevices(bool InForceRefresh = false) {}
+
+	/** Returns the channel count for the currently selected audio device */
+	virtual int32 GetDeviceChannelCount() { return 0; }
+
+	/** Accessor for the OnAudioInputDeviceChanged delegate list */
+	FOnAudioInputDeviceChanged& GetOnAudioInputDeviceChanged() { return OnAudioInputDeviceChanged;  }
+
+private:
+
+	/** Multicast notification when the user selects a new audio device */
+	FOnAudioInputDeviceChanged OnAudioInputDeviceChanged;
+};
diff --git a/Engine/Plugins/VirtualProduction/Takes/Takes.uplugin b/Engine/Plugins/VirtualProduction/Takes/Takes.uplugin
--- a/Engine/Plugins/VirtualProduction/Takes/Takes.uplugin
+++ b/Engine/Plugins/VirtualProduction/Takes/Takes.uplugin
@@ -15,6 +15,10 @@
 	"Installed": false,
 	"Plugins":
 	[
+		{
+			"Name": "AudioCapture",
+			"Enabled": true
+		},
 		{
 			"Name": "LevelSequenceEditor",
 			"Enabled": true
diff --git a/Engine/Source/Runtime/AudioCaptureImplementations/Windows/AudioCaputureWasapi/Private/AudioCaptureWasapi.cpp b/Engine/Source/Runtime/AudioCaptureImplementations/Windows/AudioCaputureWasapi/Private/AudioCaptureWasapi.cpp
--- a/Engine/Source/Runtime/AudioCaptureImplementations/Windows/AudioCaputureWasapi/Private/AudioCaptureWasapi.cpp
+++ b/Engine/Source/Runtime/AudioCaptureImplementations/Windows/AudioCaputureWasapi/Private/AudioCaptureWasapi.cpp
@@ -1,6 +1,7 @@
 // Copyright Epic Games, Inc. All Rights Reserved.
 
 #include "AudioCaptureWasapi.h"
+#include "WasapiCaptureLog.h"
 
 
 namespace Audio
@@ -40,6 +41,10 @@ namespace Audio
 				return true;
 			}
 		}
+		else
+		{
+			UE_LOG(LogAudioCaptureCore, Display, TEXT("FAudioCaptureWasapiStream::GetCaptureDeviceInfo: no default capture device found"));
+		}
 
 		return false;
 	}
@@ -87,6 +92,10 @@ namespace Audio
 				}
 			}
 		}
+		else
+		{
+			UE_LOG(LogAudioCaptureCore, Display, TEXT("FAudioCaptureWasapiStream::OpenAudioCaptureStream: no default capture device found"));
+		}
 
 		return false;
 	}
@@ -183,15 +192,12 @@ namespace Audio
 		const FWasapiDeviceEnumeration::FDeviceInfo& InDeviceInfo,
 		FWasapiAudioFormat& OutAudioFormat)
 	{
-		int32 NumChannels = InParams.NumInputChannels;
+		// Ignore InParams.NumInputChannels as WASAPI returns AUDCLNT_E_UNSUPPORTED_FORMAT 
+		// error for anything other than the device channel count (i.e. if the device 
+		// supports 2 channels, specifying 1 channel will error).
+		int32 NumChannels = InDeviceInfo.NumInputChannels;
 		int32 SampleRate = InParams.SampleRate;
 
-		// If the input params didn't specify num channels, match the device
-		if (NumChannels == InvalidDeviceChannelCount)
-		{
-			NumChannels = InDeviceInfo.NumInputChannels;
-		}
-
 		// If the input params didn't specify sample rate, match the device
 		if (SampleRate == InvalidDeviceSampleRate)
 		{
diff --git a/Engine/Source/Runtime/AudioCaptureImplementations/Windows/AudioCaputureWasapi/Private/WasapiDeviceEnumeration.cpp b/Engine/Source/Runtime/AudioCaptureImplementations/Windows/AudioCaputureWasapi/Private/WasapiDeviceEnumeration.cpp
--- a/Engine/Source/Runtime/AudioCaptureImplementations/Windows/AudioCaputureWasapi/Private/WasapiDeviceEnumeration.cpp
+++ b/Engine/Source/Runtime/AudioCaptureImplementations/Windows/AudioCaputureWasapi/Private/WasapiDeviceEnumeration.cpp
@@ -210,10 +210,10 @@ namespace Audio
 				FDeviceInfo DeviceInfo;
 				if (GetDeviceProperties(Device, DeviceInfo))
 				{
-					UE_LOG(LogAudioCaptureCore, Display, TEXT("Device %d: \"%s\" (%s)"), DeviceIndex,
+					UE_LOG(LogAudioCaptureCore, Verbose, TEXT("Device %d: \"%s\" (%s)"), DeviceIndex,
 						DeviceInfo.FriendlyName.GetCharArray().GetData(), DeviceInfo.DeviceId.GetCharArray().GetData());
 
-					UE_LOG(LogAudioCaptureCore, Display, TEXT("\tBPS: %d, SR: %d"), DeviceInfo.BitsPerSample, DeviceInfo.PreferredSampleRate);
+					UE_LOG(LogAudioCaptureCore, Verbose, TEXT("\tBPS: %d, SR: %d"), DeviceInfo.BitsPerSample, DeviceInfo.PreferredSampleRate);
 
 					OutDevices.Emplace(MoveTemp(DeviceInfo));
 				}
