Revert pervious change to channel count when recording audio.

diff --git a/Engine/Plugins/Runtime/AudioCapture/Source/AudioCaptureEditor/Private/AudioRecordingManager.cpp b/Engine/Plugins/Runtime/AudioCapture/Source/AudioCaptureEditor/Private/AudioRecordingManager.cpp
--- a/Engine/Plugins/Runtime/AudioCapture/Source/AudioCaptureEditor/Private/AudioRecordingManager.cpp
+++ b/Engine/Plugins/Runtime/AudioCapture/Source/AudioCaptureEditor/Private/AudioRecordingManager.cpp
@@ -76,7 +76,7 @@ void FAudioRecordingManager::StartRecording(const FRecordingSettings& InSettings
 		}
 
 		RecordingSampleRate = DeviceInfo.PreferredSampleRate;
-		NumInputChannels = FMath::Min(DeviceInfo.InputChannels, Settings.NumRecordChannels);
+		NumInputChannels = DeviceInfo.InputChannels;
 
 		if (NumInputChannels == 0)
 		{
@@ -85,7 +85,7 @@ void FAudioRecordingManager::StartRecording(const FRecordingSettings& InSettings
 		}
 
 		// Initialize Sound Waves array
-		RecordedSoundWaves.Init(nullptr, DeviceInfo.InputChannels);
+		RecordedSoundWaves.Init(nullptr, NumInputChannels);
 
 		// Reserve enough space in our current recording buffer for 10 seconds of audio to prevent slowing down due to allocations
 		if (Settings.RecordingDuration > 0.0f)
