Fix for stale SoundWaveData on newly created USoundWave objects in Take Recorder. Regression caused by fix for UE-180245.

diff --git a/Engine/Plugins/Runtime/AudioCapture/Source/AudioCaptureEditor/Private/AudioRecordingManager.cpp b/Engine/Plugins/Runtime/AudioCapture/Source/AudioCaptureEditor/Private/AudioRecordingManager.cpp
--- a/Engine/Plugins/Runtime/AudioCapture/Source/AudioCaptureEditor/Private/AudioRecordingManager.cpp
+++ b/Engine/Plugins/Runtime/AudioCapture/Source/AudioCaptureEditor/Private/AudioRecordingManager.cpp
@@ -408,10 +408,12 @@ TObjectPtr<USoundWave> FAudioRecordingManager::CreateSoundWaveAsset(const FRecor
 		NewSoundWave->SetSoundAssetCompressionType(ESoundAssetCompressionType::BinkAudio);
 		NewSoundWave->SetTimecodeInfo(GetTimecodeInfo(InSourceSettings));
 
+		// Initialize SoundWaveData so that it is synchronized with the owning object
+		// (note that for serialized cases, this happens in PostLoad(), PostImport(), etc.)
+		NewSoundWave->SoundWaveDataPtr->InitializeDataFromSoundWave(*NewSoundWave);
+
 		if (bCreatedPackage)
 		{
-			//GEditor->GetEditorSubsystem<UImportSubsystem>()->BroadcastAssetPostImport(this, NewSoundWave);
-
 			FAssetRegistryModule::AssetCreated(NewSoundWave);
 			NewSoundWave->MarkPackageDirty();
 
