Fixing overuse of memzeros

Saves 2% frametime on test project

diff --git a/Engine/Plugins/Media/PixelStreaming/Source/PixelStreaming/Private/PixelStreamingAudioComponent.cpp b/Engine/Plugins/Media/PixelStreaming/Source/PixelStreaming/Private/PixelStreamingAudioComponent.cpp
--- a/Engine/Plugins/Media/PixelStreaming/Source/PixelStreaming/Private/PixelStreamingAudioComponent.cpp
+++ b/Engine/Plugins/Media/PixelStreaming/Source/PixelStreaming/Private/PixelStreamingAudioComponent.cpp
@@ -239,6 +239,7 @@ int32 FWebRTCSoundGenerator::OnGenerateAudio(float* OutAudio, int32 NumSamples)
 	// Not listening to peer, return zero'd buffer.
 	if (!bShouldGenerateAudio || Buffer.Num() == 0)
 	{
+		FMemory::Memzero(OutAudio, NumSamples * sizeof(float));
 		return NumSamples;
 	}
 
@@ -251,6 +252,7 @@ int32 FWebRTCSoundGenerator::OnGenerateAudio(float* OutAudio, int32 NumSamples)
 		// Copy from local buffer into OutAudio if we have enough samples
 		for (int i = 0; i < NumSamplesToCopy; i++)
 		{
+			//TODO: vectorize!
 			*OutAudio = Buffer[i] / 32767.0f;
 			OutAudio++;
 		}
diff --git a/Engine/Plugins/Runtime/AudioCapture/Source/AudioCapture/Private/AudioCaptureComponent.cpp b/Engine/Plugins/Runtime/AudioCapture/Source/AudioCapture/Private/AudioCaptureComponent.cpp
--- a/Engine/Plugins/Runtime/AudioCapture/Source/AudioCapture/Private/AudioCaptureComponent.cpp
+++ b/Engine/Plugins/Runtime/AudioCapture/Source/AudioCapture/Private/AudioCaptureComponent.cpp
@@ -130,6 +130,7 @@ int32 UAudioCaptureComponent::OnGenerateAudio(float* OutAudio, int32 NumSamples)
 	if (!bIsStreamOpen || !CaptureSynth.IsStreamOpen() || !CaptureSynth.IsCapturing())
 	{
 		// Just return NumSamples, which uses zero'd buffer
+		FMemory::Memzero(OutAudio, NumSamples * sizeof(float));
 		return NumSamples;
 	}
 	int32 OutputSamplesGenerated = 0;
@@ -139,6 +140,7 @@ int32 UAudioCaptureComponent::OnGenerateAudio(float* OutAudio, int32 NumSamples)
 		//Clear the CaptureSynth's data, too
 		CaptureSynth.GetAudioData(CaptureAudioData);
 		CaptureAudioData.Reset();
+		FMemory::Memzero(OutAudio, NumSamples * sizeof(float));
 		return NumSamples;
 	}
 	if (CapturedAudioDataSamples > 0 || CaptureSynth.GetNumSamplesEnqueued() > 1024)
@@ -192,7 +194,13 @@ int32 UAudioCaptureComponent::OnGenerateAudio(float* OutAudio, int32 NumSamples)
 	else
 	{
 		// Say we generated the full samples, this will result in silence
+		FMemory::Memzero(OutAudio, NumSamples * sizeof(float));
 		OutputSamplesGenerated = NumSamples;
 	}
+
+	if (OutputSamplesGenerated < NumSamples)
+	{
+		FMemory::Memzero(&OutAudio[OutputSamplesGenerated], (NumSamples - OutputSamplesGenerated) * sizeof(float));
+	}
 	return OutputSamplesGenerated;
 }
\ No newline at end of file
diff --git a/Engine/Plugins/Runtime/Synthesis/Source/Synthesis/Private/SynthComponentWaveTable.cpp b/Engine/Plugins/Runtime/Synthesis/Source/Synthesis/Private/SynthComponentWaveTable.cpp
--- a/Engine/Plugins/Runtime/Synthesis/Source/Synthesis/Private/SynthComponentWaveTable.cpp
+++ b/Engine/Plugins/Runtime/Synthesis/Source/Synthesis/Private/SynthComponentWaveTable.cpp
@@ -184,10 +184,7 @@ int32 USynthSamplePlayer::OnGenerateAudio(float* OutAudio, int32 NumSamples)
 	}
 	else
 	{
-		for (int32 Sample = 0; Sample < NumSamples; ++Sample)
-		{
-			OutAudio[Sample] = 0.0f;
-		}
+		FMemory::Memzero(OutAudio, NumSamples * sizeof(float));
 	}
 	return NumSamples;
 }
diff --git a/Engine/Source/Runtime/AudioMixer/Private/AudioMixerSourceBuffer.cpp b/Engine/Source/Runtime/AudioMixer/Private/AudioMixerSourceBuffer.cpp
--- a/Engine/Source/Runtime/AudioMixer/Private/AudioMixerSourceBuffer.cpp
+++ b/Engine/Source/Runtime/AudioMixer/Private/AudioMixerSourceBuffer.cpp
@@ -39,17 +39,11 @@ namespace Audio
 			}
 
 			// Zero out the rest of the buffer
-			while (Sample < NumSampleToGet)
-			{
-				OutBufferPtr[Sample++] = 0.0f;
-			}
+			FMemory::Memzero(&OutBufferPtr[Sample], (NumSampleToGet - Sample) * sizeof(float));
 		}
 		else
 		{
-			for (uint32 Sample = 0; Sample < NumSampleToGet; ++Sample)
-			{
-				OutBufferPtr[Sample] = 0.0f;
-			}
+			FMemory::Memzero(OutBufferPtr, NumSampleToGet * sizeof(float));
 		}
 
 		// If the current sample is greater or equal to num samples we hit the end of the buffer
@@ -341,7 +335,7 @@ namespace Audio
 		const int32 MaxSamples = MONO_PCM_BUFFER_SAMPLES * NumChannels;
 
 		SourceVoiceBuffers[BufferIndex]->AudioData.Reset();
-		SourceVoiceBuffers[BufferIndex]->AudioData.AddZeroed(MaxSamples);
+		SourceVoiceBuffers[BufferIndex]->AudioData.AddUninitialized(MaxSamples);
 
 		if (bProcedural)
 		{
@@ -384,6 +378,8 @@ namespace Audio
 		// Handle the case that the decoder has an error and can't continue.
 		if (InDecoder && InDecoder->HasError())
 		{
+			FMemory::Memzero(SourceVoiceBuffers[BufferIndex]->AudioData.GetData(), MaxSamples * sizeof(float));
+
 			FScopeTryLock Lock(&SoundWaveCritSec);
 			if (Lock.IsLocked() && SoundWave)
 			{
diff --git a/Engine/Source/Runtime/AudioMixer/Private/AudioMixerSourceDecode.cpp b/Engine/Source/Runtime/AudioMixer/Private/AudioMixerSourceDecode.cpp
--- a/Engine/Source/Runtime/AudioMixer/Private/AudioMixerSourceDecode.cpp
+++ b/Engine/Source/Runtime/AudioMixer/Private/AudioMixerSourceDecode.cpp
@@ -93,8 +93,7 @@ public:
 				QUICK_SCOPE_CYCLE_COUNTER(STAT_FAsyncDecodeWorker_Procedural);
 				if (ProceduralTaskData.SoundGenerator.IsValid())
 				{
-					// Pre-zero the buffer before calling into the generator code as a convenience to implementers
-					FMemory::Memzero(ProceduralTaskData.AudioData, ProceduralTaskData.NumSamples * sizeof(float));
+					// Generators are responsible to zero memory in case they can't generate the requested amount of samples
 					ProceduralResult.NumSamplesWritten = ProceduralTaskData.SoundGenerator->GetNextBuffer(ProceduralTaskData.AudioData, ProceduralTaskData.NumSamples);
 					ProceduralResult.bIsFinished = ProceduralTaskData.SoundGenerator->IsFinished();
 				}
diff --git a/Engine/Source/Runtime/Engine/Classes/Sound/SoundGenerator.h b/Engine/Source/Runtime/Engine/Classes/Sound/SoundGenerator.h
--- a/Engine/Source/Runtime/Engine/Classes/Sound/SoundGenerator.h
+++ b/Engine/Source/Runtime/Engine/Classes/Sound/SoundGenerator.h
@@ -74,6 +74,7 @@ class FSoundGeneratorNull : public ISoundGenerator
 public:
 	virtual int32 OnGenerateAudio(float* OutAudio, int32 NumSamples) override
 	{
+		FMemory::Memzero(OutAudio, NumSamples * sizeof(float));
 		return NumSamples;
 	}
 };
diff --git a/Engine/Source/Runtime/SignalProcessing/Private/Granulator.cpp b/Engine/Source/Runtime/SignalProcessing/Private/Granulator.cpp
--- a/Engine/Source/Runtime/SignalProcessing/Private/Granulator.cpp
+++ b/Engine/Source/Runtime/SignalProcessing/Private/Granulator.cpp
@@ -776,6 +776,8 @@ namespace Audio
 
 	void FGranularSynth::Generate(float* OutAudiobuffer, const int32 NumFrames)
 	{
+		FMemory::Memzero(OutAudiobuffer, NumFrames * sizeof(float));
+
 		if (SampleBuffer.GetData() == nullptr)
 		{
 			return;
diff --git a/Engine/Source/Runtime/SignalProcessing/Private/SampleBufferReader.cpp b/Engine/Source/Runtime/SignalProcessing/Private/SampleBufferReader.cpp
--- a/Engine/Source/Runtime/SignalProcessing/Private/SampleBufferReader.cpp
+++ b/Engine/Source/Runtime/SignalProcessing/Private/SampleBufferReader.cpp
@@ -210,13 +210,10 @@ namespace Audio
 	bool FSampleBufferReader::Generate(float* OutAudioBuffer, const int32 NumFrames, const int32 OutChannels, const bool bInWrap)
 	{
 		// Don't have a buffer yet, so fill in zeros, say we're not done
+		const int32 NumSamples = NumFrames * OutChannels;
 		if (!HasBuffer() || bIsFinished)
 		{
-			int32 NumSamples = NumFrames * OutChannels;
-			for (int32 i = 0; i < NumSamples; ++i)
-			{
-				OutAudioBuffer[i] = 0.0f;
-			}
+			FMemory::Memzero(OutAudioBuffer, NumSamples * sizeof(float));
 			return false;
 		}	
 
@@ -346,6 +343,11 @@ namespace Audio
 			}
 		}
 #endif
+		if (OutSampleIndex < NumSamples)
+		{
+			FMemory::Memzero(&OutAudioBuffer[OutSampleIndex], (NumSamples - OutSampleIndex) * sizeof(float));
+		}
+
 		return bIsFinished;
 	}
 
