UE-219861 VisionOS: audio capture support
-Add visionOS platform support

diff --git a/Engine/Platforms/VisionOS/Plugins/Runtime/AudioCapture/AudioCapture_VisionOS.uplugin b/Engine/Platforms/VisionOS/Plugins/Runtime/AudioCapture/AudioCapture_VisionOS.uplugin
new file mode 100644
--- /dev/null
+++ b/Engine/Platforms/VisionOS/Plugins/Runtime/AudioCapture/AudioCapture_VisionOS.uplugin
@@ -0,0 +1,4 @@
+{
+	"FileVersion" : 3,
+	"bIsPluginExtension" : true
+}
diff --git a/Engine/Platforms/VisionOS/Plugins/Runtime/AudioCapture/Source/AudioCapture_VisionOS.Build.cs b/Engine/Platforms/VisionOS/Plugins/Runtime/AudioCapture/Source/AudioCapture_VisionOS.Build.cs
new file mode 100644
--- /dev/null
+++ b/Engine/Platforms/VisionOS/Plugins/Runtime/AudioCapture/Source/AudioCapture_VisionOS.Build.cs
@@ -0,0 +1,15 @@
+// Copyright Epic Games, Inc. All Rights Reserved.
+
+namespace UnrealBuildTool.Rules
+{
+	public class AudioCapture_VisionOS : AudioCapture
+	{
+		public AudioCapture_VisionOS(ReadOnlyTargetRules Target) : base(Target)
+		{
+			if (Target.Platform == UnrealTargetPlatform.VisionOS)
+            {
+                PrivateDependencyModuleNames.Add("AudioCaptureAudioUnit");
+            }
+        }
+	}
+}
\ No newline at end of file
