[Audio] Remove transient buffer allocation from UAudioCaptureComponent CDO

#rnx

diff --git a/Engine/Plugins/Runtime/AudioCapture/Source/AudioCapture/Private/AudioCaptureComponent.cpp b/Engine/Plugins/Runtime/AudioCapture/Source/AudioCapture/Private/AudioCaptureComponent.cpp
--- a/Engine/Plugins/Runtime/AudioCapture/Source/AudioCapture/Private/AudioCaptureComponent.cpp
+++ b/Engine/Plugins/Runtime/AudioCapture/Source/AudioCapture/Private/AudioCaptureComponent.cpp
@@ -16,7 +16,10 @@ UAudioCaptureComponent::UAudioCaptureComponent(const FObjectInitializer& ObjectI
 	bIsDestroying = false;
 	bIsNotReadyForForFinishDestroy = false;
 	bIsStreamOpen = false;
-	CaptureAudioData.Reserve(2 * 48000 * 5);
+	if (!HasAnyFlags(RF_ClassDefaultObject))
+	{
+		CaptureAudioData.Reserve(2 * 48000 * 5);
+	}
 }
 
 bool UAudioCaptureComponent::Init(int32& SampleRate)
@@ -88,6 +91,11 @@ void UAudioCaptureComponent::FinishDestroy()
 	bIsStreamOpen = false;
 }
 
+void UAudioCaptureComponent::GetResourceSizeEx(FResourceSizeEx& CumulativeResourceSize)
+{
+	CumulativeResourceSize.AddDedicatedSystemMemoryBytes(CaptureAudioData.GetAllocatedSize());
+}
+
 void UAudioCaptureComponent::OnBeginGenerate()
 {
 	CapturedAudioDataSamples = 0;
diff --git a/Engine/Plugins/Runtime/AudioCapture/Source/AudioCapture/Public/AudioCaptureComponent.h b/Engine/Plugins/Runtime/AudioCapture/Source/AudioCapture/Public/AudioCaptureComponent.h
--- a/Engine/Plugins/Runtime/AudioCapture/Source/AudioCapture/Public/AudioCaptureComponent.h
+++ b/Engine/Plugins/Runtime/AudioCapture/Source/AudioCapture/Public/AudioCaptureComponent.h
@@ -27,6 +27,7 @@ protected:
 	virtual void BeginDestroy();
 	virtual bool IsReadyForFinishDestroy() override;
 	virtual void FinishDestroy() override;
+	virtual void GetResourceSizeEx(FResourceSizeEx& CumulativeResourceSize) override;
 	//~ End UObject interface
 
 public:
